UAF: Rename/move plugins

diff --git a/Engine/Plugins/Experimental/MovieScenePoseSearchTracks/MovieScenePoseSearchTracks.uplugin b/Engine/Plugins/Experimental/MovieScenePoseSearchTracks/MovieScenePoseSearchTracks.uplugin
--- a/Engine/Plugins/Experimental/MovieScenePoseSearchTracks/MovieScenePoseSearchTracks.uplugin
+++ b/Engine/Plugins/Experimental/MovieScenePoseSearchTracks/MovieScenePoseSearchTracks.uplugin
@@ -22,7 +22,7 @@
 			"Enabled": true
 		},
 		{
-			"Name": "AnimNextPoseSearch",
+			"Name": "UAFPoseSearch",
 			"Enabled": true
 		}
 	],
diff --git a/Engine/Plugins/Experimental/MovieScenePoseSearchTracks/Source/MovieScenePoseSearchTracks/MovieScenePoseSearchTracks.Build.cs b/Engine/Plugins/Experimental/MovieScenePoseSearchTracks/Source/MovieScenePoseSearchTracks/MovieScenePoseSearchTracks.Build.cs
--- a/Engine/Plugins/Experimental/MovieScenePoseSearchTracks/Source/MovieScenePoseSearchTracks/MovieScenePoseSearchTracks.Build.cs
+++ b/Engine/Plugins/Experimental/MovieScenePoseSearchTracks/Source/MovieScenePoseSearchTracks/MovieScenePoseSearchTracks.Build.cs
@@ -16,7 +16,7 @@ namespace UnrealBuildTool.Rules
 					"MovieSceneTracks",
 					"MovieSceneAnimMixer",
 					"UAF",
-					"AnimNextPoseSearch",
+					"UAFPoseSearch",
 					"UAFAnimGraph",
 					"AnimGraphRuntime",
 					"PoseSearch"
diff --git a/Engine/Plugins/Experimental/MovieScenePoseSearchTracks/Source/MovieScenePoseSearchTracksEditor/MovieScenePoseSearchTracksEditor.Build.cs b/Engine/Plugins/Experimental/MovieScenePoseSearchTracks/Source/MovieScenePoseSearchTracksEditor/MovieScenePoseSearchTracksEditor.Build.cs
--- a/Engine/Plugins/Experimental/MovieScenePoseSearchTracks/Source/MovieScenePoseSearchTracksEditor/MovieScenePoseSearchTracksEditor.Build.cs
+++ b/Engine/Plugins/Experimental/MovieScenePoseSearchTracks/Source/MovieScenePoseSearchTracksEditor/MovieScenePoseSearchTracksEditor.Build.cs
@@ -19,7 +19,7 @@ namespace UnrealBuildTool.Rules
 					"AnimGraphRuntime",
 					"UAF",
 					"UAFAnimGraph",
-					"AnimNextPoseSearch",
+					"UAFPoseSearch",
 					"Sequencer",
 					"SequencerCore",
 					"SlateCore",
diff --git a/Engine/Plugins/Experimental/UAF/UAF/Source/UAFTestSuite/Private/AnimNextVariablesTest.cpp b/Engine/Plugins/Experimental/UAF/UAF/Source/UAFTestSuite/Private/AnimNextVariablesTest.cpp
--- a/Engine/Plugins/Experimental/UAF/UAF/Source/UAFTestSuite/Private/AnimNextVariablesTest.cpp
+++ b/Engine/Plugins/Experimental/UAF/UAF/Source/UAFTestSuite/Private/AnimNextVariablesTest.cpp
@@ -331,7 +331,7 @@ bool FVariables_UOLBindings::RunTest(const FString& InParameters)
 	BindingDataA.Type = FAnimNextUniversalObjectLocatorBindingType::Property;
 	BindingDataA.Property = UAnimNextTestFuncLib::StaticClass()->FindPropertyByName(GET_MEMBER_NAME_CHECKED(UAnimNextTestFuncLib, ValueA));
 	UE_RETURN_ON_ERROR(BindingDataA.Property != nullptr, TEXT("Could not find property."));
-	UE::UniversalObjectLocator::FParseStringResult ResultA = BindingDataA.Locator.TryParseString(TEXT("uobj://animobjfunc?&payload0=/Script/AnimNextTestSuite.AnimNextTestFuncLib:GetObj"), ParseStringParams);
+	UE::UniversalObjectLocator::FParseStringResult ResultA = BindingDataA.Locator.TryParseString(TEXT("uobj://animobjfunc?&payload0=/Script/UAFTestSuite.AnimNextTestFuncLib:GetObj"), ParseStringParams);
 	UE_RETURN_ON_ERROR(ResultA.bSuccess, TEXT("Could not parse UOL."));
 	PropertyEntry->SetBinding(MoveTemp(BindingA));
 
@@ -342,7 +342,7 @@ bool FVariables_UOLBindings::RunTest(const FString& InParameters)
 	BindingDataB.Type = FAnimNextUniversalObjectLocatorBindingType::Function;
 	BindingDataB.Function = UAnimNextTestFuncLib::StaticClass()->FindFunctionByName(GET_FUNCTION_NAME_CHECKED(UAnimNextTestFuncLib, GetValueB));
 	UE_RETURN_ON_ERROR(BindingDataB.Function != nullptr, TEXT("Could not find function."));
-	UE::UniversalObjectLocator::FParseStringResult ResultB = BindingDataB.Locator.TryParseString(TEXT("uobj://animobjfunc?&payload0=/Script/AnimNextTestSuite.AnimNextTestFuncLib:GetObj"), ParseStringParams);
+	UE::UniversalObjectLocator::FParseStringResult ResultB = BindingDataB.Locator.TryParseString(TEXT("uobj://animobjfunc?&payload0=/Script/UAFTestSuite.AnimNextTestFuncLib:GetObj"), ParseStringParams);
 	UE_RETURN_ON_ERROR(ResultB.bSuccess, TEXT("Could not parse UOL."));
 	ObjectAccessorEntry->SetBinding(MoveTemp(BindingB));
 
@@ -353,7 +353,7 @@ bool FVariables_UOLBindings::RunTest(const FString& InParameters)
 	BindingDataC.Type = FAnimNextUniversalObjectLocatorBindingType::HoistedFunction;
 	BindingDataC.Function = UAnimNextTestFuncLib::StaticClass()->FindFunctionByName(GET_FUNCTION_NAME_CHECKED(UAnimNextTestFuncLib, GetValueC));
 	UE_RETURN_ON_ERROR(BindingDataC.Function != nullptr, TEXT("Could not find function."));
-	UE::UniversalObjectLocator::FParseStringResult ResultC = BindingDataC.Locator.TryParseString(TEXT("uobj://animobjfunc?&payload0=/Script/AnimNextTestSuite.AnimNextTestFuncLib:GetObj"), ParseStringParams);
+	UE::UniversalObjectLocator::FParseStringResult ResultC = BindingDataC.Locator.TryParseString(TEXT("uobj://animobjfunc?&payload0=/Script/UAFTestSuite.AnimNextTestFuncLib:GetObj"), ParseStringParams);
 	UE_RETURN_ON_ERROR(ResultC.bSuccess, TEXT("Could not parse UOL."));
 	HoistedAccessorEntry->SetBinding(MoveTemp(BindingC));
 
diff --git a/Engine/Plugins/Experimental/UAF/UAFControlRig/Config/BaseUAFControlRig.ini b/Engine/Plugins/Experimental/UAF/UAFControlRig/Config/BaseUAFControlRig.ini
new file mode 100644
--- /dev/null
+++ b/Engine/Plugins/Experimental/UAF/UAFControlRig/Config/BaseUAFControlRig.ini
@@ -0,0 +1,9 @@
+﻿[CoreRedirects]
+; AnimNextChooser plugin -> UAFChooser plugin
++PackageRedirects=(OldName="/Script/AnimNextControlRig", NewName="/Script/UAFControlRig")
+
++StructRedirects=(OldName="/Script/AnimNextControlRig.ControlRigEventName", NewName="/Script/UAFControlRig.ControlRigEventName")
++StructRedirects=(OldName="/Script/AnimNextControlRig.ControlRigExposedProperty", NewName="/Script/UAFControlRig.ControlRigExposedProperty")
++StructRedirects=(OldName="/Script/AnimNextControlRig.ControlRigTraitSharedData", NewName="/Script/UAFControlRig.ControlRigTraitSharedData")
++StructRedirects=(OldName="/Script/AnimNextControlRig.AnimNextControlRigTask", NewName="/Script/UAFControlRig.AnimNextControlRigTask")
+
diff --git a/Engine/Plugins/Experimental/UAF/UAFControlRig/Source/UAFControlRig/Private/ControlRigTask.h b/Engine/Plugins/Experimental/UAF/UAFControlRig/Source/UAFControlRig/Private/ControlRigTask.h
new file mode 100644
--- /dev/null
+++ b/Engine/Plugins/Experimental/UAF/UAFControlRig/Source/UAFControlRig/Private/ControlRigTask.h
@@ -0,0 +1,44 @@
+// Copyright Epic Games, Inc. All Rights Reserved.
+
+#pragma once
+
+#include "CoreMinimal.h"
+#include "EvaluationVM/EvaluationTask.h"
+#include "Animation/AnimCurveTypes.h"
+#include "LODPose.h"
+#include "ControlRigTrait.h"
+#include "AnimNextStats.h"
+// --- ---
+#include "ControlRigTask.generated.h"
+
+namespace UE::UAF
+{
+struct FKeyframeState;
+}
+
+DECLARE_CYCLE_STAT_EXTERN(TEXT("UAF Task: ControlRig"), STAT_AnimNext_Task_ControlRig, STATGROUP_AnimNext, UAFCONTROLRIG_API);
+
+USTRUCT()
+struct UAFCONTROLRIG_API FAnimNextControlRigTask : public FAnimNextEvaluationTask
+{
+	GENERATED_BODY()
+
+	DECLARE_ANIM_EVALUATION_TASK(FAnimNextControlRigTask)
+
+	static FAnimNextControlRigTask Make(const UE::UAF::FControlRigTrait::FSharedData* SharedData, UE::UAF::FControlRigTrait::FInstanceData* InstanceData);
+
+	// Task entry point
+	virtual void Execute(UE::UAF::FEvaluationVM& VM) const override;
+
+private:
+	void ExecuteControlRig(UE::UAF::FEvaluationVM& VM, UE::UAF::FKeyframeState& KeyFrameState, UControlRig* ControlRig) const;
+	void UpdateInput(UE::UAF::FEvaluationVM& VM, UE::UAF::FKeyframeState& KeyFrameState, UControlRig* ControlRig, UE::UAF::FKeyframeState& InOutput) const;
+	void UpdateOutput(UE::UAF::FEvaluationVM& VM, UE::UAF::FKeyframeState& KeyFrameState, UControlRig* ControlRig, UE::UAF::FKeyframeState& InOutput) const;
+
+	bool CanExecute(const UControlRig* ControlRig) const;
+
+	void QueueControlRigDrawInstructions(UControlRig* ControlRig, FRigVMDrawInterface* DebugDrawInterface, const FTransform& ComponentTransform) const;
+
+	const UE::UAF::FControlRigTrait::FSharedData* SharedData = nullptr;
+	UE::UAF::FControlRigTrait::FInstanceData* InstanceData = nullptr;
+};
diff --git a/Engine/Plugins/Experimental/UAF/UAFPoseSearch/Config/BaseUAFPoseSearch.ini b/Engine/Plugins/Experimental/UAF/UAFPoseSearch/Config/BaseUAFPoseSearch.ini
new file mode 100644
--- /dev/null
+++ b/Engine/Plugins/Experimental/UAF/UAFPoseSearch/Config/BaseUAFPoseSearch.ini
@@ -0,0 +1,18 @@
+﻿[CoreRedirects]
+; AnimNextPoseSearch plugin -> UAFPoseSearch plugin
++PackageRedirects=(OldName="/Script/AnimNextPoseSearch", NewName="/Script/UAFPoseSearch")
+
++StructRedirects=(OldName="/Script/AnimNextPoseSearch.AnimNextHistoryCollectorTraitSharedData", NewName="/Script/UAFPoseSearch.AnimNextHistoryCollectorTraitSharedData")
++StructRedirects=(OldName="/Script/AnimNextPoseSearch.AnimNextHistoryCollectorTask", NewName="/Script/UAFPoseSearch.AnimNextHistoryCollectorTask")
++StructRedirects=(OldName="/Script/AnimNextPoseSearch.AnimNextHistoryCollectorPreEvaluateTask", NewName="/Script/UAFPoseSearch.AnimNextHistoryCollectorPreEvaluateTask")
++StructRedirects=(OldName="/Script/AnimNextPoseSearch.MotionMatchingTraitSharedData", NewName="/Script/UAFPoseSearch.MotionMatchingTraitSharedData")
++StructRedirects=(OldName="/Script/AnimNextPoseSearch.AnimNextMotionMatchingTask", NewName="/Script/UAFPoseSearch.AnimNextMotionMatchingTask")
++StructRedirects=(OldName="/Script/AnimNextPoseSearch.PoseHistoryAnimProperty", NewName="/Script/UAFPoseSearch.PoseHistoryAnimProperty")
++StructRedirects=(OldName="/Script/AnimNextPoseSearch.PoseSearchResultEmulatorTraitSharedData", NewName="/Script/UAFPoseSearch.PoseSearchResultEmulatorTraitSharedData")
++StructRedirects=(OldName="/Script/AnimNextPoseSearch.AnimNextPoseSearchResultEmulatorTask", NewName="/Script/UAFPoseSearch.AnimNextPoseSearchResultEmulatorTask")
++StructRedirects=(OldName="/Script/AnimNextPoseSearch.RigUnit_DebugDrawTrajectory", NewName="/Script/UAFPoseSearch.RigUnit_DebugDrawTrajectory")
++StructRedirects=(OldName="/Script/AnimNextPoseSearch.RigUnit_MultiAnimGetAnimationAsset", NewName="/Script/UAFPoseSearch.RigUnit_MultiAnimGetAnimationAsset")
++StructRedirects=(OldName="/Script/AnimNextPoseSearch.RigUnit_PoseSearchDatabaseGetTags", NewName="/Script/UAFPoseSearch.RigUnit_PoseSearchDatabaseGetTags")
++StructRedirects=(OldName="/Script/AnimNextPoseSearch.RigUnit_PoseSearchResultGetSelectedDatabase", NewName="/Script/UAFPoseSearch.RigUnit_PoseSearchResultGetSelectedDatabase")
+
++ClassRedirects=(OldName="/Script/AnimNextPoseSearch.PoseSearchFeatureChannel_DistanceFromAnimNextVar", NewName="/Script/UAFPoseSearch.PoseSearchFeatureChannel_DistanceFromAnimNextVar")
\ No newline at end of file
diff --git a/Engine/Plugins/Experimental/UAF/UAFStateTree/Config/BaseUAFStateTree.ini b/Engine/Plugins/Experimental/UAF/UAFStateTree/Config/BaseUAFStateTree.ini
new file mode 100644
--- /dev/null
+++ b/Engine/Plugins/Experimental/UAF/UAFStateTree/Config/BaseUAFStateTree.ini
@@ -0,0 +1,28 @@
+﻿[CoreRedirects]
+; AnimNextStateTree plugin -> UAFStateTree plugin
++PackageRedirects=(OldName="/Script/AnimNextStateTree", NewName="/Script/UAFStateTree")
+
++StructRedirects=(OldName="/Script/AnimNextStateTree.AnimNextStateTreeEvaluatorBase", NewName="/Script/UAFStateTree.AnimNextStateTreeEvaluatorBase")
++StructRedirects=(OldName="/Script/AnimNextStateTree.AnimNextStateTreeTaskBase", NewName="/Script/UAFStateTree.AnimNextStateTreeTaskBase")
++StructRedirects=(OldName="/Script/AnimNextStateTree.AnimNextStateTreeTraitSharedData", NewName="/Script/UAFStateTree.AnimNextStateTreeTraitSharedData")
++StructRedirects=(OldName="/Script/AnimNextStateTree.AnimNextStateTreeRigVMConditionInstanceData", NewName="/Script/UAFStateTree.AnimNextStateTreeRigVMConditionInstanceData")
++StructRedirects=(OldName="/Script/AnimNextStateTree.AnimNextStateTreeRigVMConditionBase", NewName="/Script/UAFStateTree.AnimNextStateTreeRigVMConditionBase")
++StructRedirects=(OldName="/Script/AnimNextStateTree.AnimNextGraphInstanceTaskInstanceData", NewName="/Script/UAFStateTree.AnimNextGraphInstanceTaskInstanceData")
++StructRedirects=(OldName="/Script/AnimNextStateTree.AnimNextStateTreeGraphInstanceTask", NewName="/Script/UAFStateTree.AnimNextStateTreeGraphInstanceTask")
++StructRedirects=(OldName="/Script/AnimNextStateTree.AnimNextStateTreeRigVMTaskInstanceData", NewName="/Script/UAFStateTree.AnimNextStateTreeRigVMTaskInstanceData")
++StructRedirects=(OldName="/Script/AnimNextStateTree.AnimNextStateTreeRigVMTaskBase", NewName="/Script/UAFStateTree.AnimNextStateTreeRigVMTaskBase")
++StructRedirects=(OldName="/Script/AnimNextStateTree.AnimNextStateTreeTraitContext", NewName="/Script/UAFStateTree.AnimNextStateTreeTraitContext")
+
++ClassRedirects=(OldName="/Script/AnimNextStateTree.AnimNextStateTree", NewName="/Script/UAFStateTree.AnimNextStateTree")
++ClassRedirects=(OldName="/Script/AnimNextStateTree.StateTreeAnimNextSchema", NewName="/Script/UAFStateTree.StateTreeAnimNextSchema")
+
++StructRedirects=(OldName="/Script/AnimNextStateTreeUncookedOnly.AnimNextStateTreeOutlinerData", NewName="/Script/UAFStateTreeUncookedOnly.AnimNextStateTreeOutlinerData")
++StructRedirects=(OldName="/Script/AnimNextStateTreeUncookedOnly.AnimNextStateTreeStateOutlinerData", NewName="/Script/UAFStateTreeUncookedOnly.AnimNextStateTreeStateOutlinerData")
+
++ClassRedirects=(OldName="/Script/AnimNextStateTreeUncookedOnly.AnimNextStateTreeWorkspaceAssetUserData", NewName="/Script/UAFStateTreeUncookedOnly.AnimNextStateTreeWorkspaceAssetUserData")
++ClassRedirects=(OldName="/Script/AnimNextStateTreeUncookedOnly.AnimNextStateTree_EditorData", NewName="/Script/UAFStateTreeUncookedOnly.AnimNextStateTree_EditorData")
++ClassRedirects=(OldName="/Script/AnimNextStateTreeUncookedOnly.AnimNextStateTreeTreeEditorData", NewName="/Script/UAFStateTreeUncookedOnly.AnimNextStateTreeTreeEditorData")
++ClassRedirects=(OldName="/Script/AnimNextStateTreeUncookedOnly.AnimNextStateTreeFunctionLibraryHelper", NewName="/Script/UAFStateTreeUncookedOnly.AnimNextStateTreeFunctionLibraryHelper")
+
++ClassRedirects=(OldName="/Script/AnimNextStateTreeEditor.AnimNextStateTreeFactory", NewName="/Script/UAFStateTreeEditor.AnimNextStateTreeFactory")
++ClassRedirects=(OldName="/Script/AnimNextStateTreeEditor.AssetDefinition_AnimNextStateTree", NewName="/Script/UAFStateTreeEditor.AssetDefinition_AnimNextStateTree")
diff --git a/Engine/Plugins/Experimental/UAF/UAFStateTree/Source/UAFStateTree/Internal/AnimNextStateTreeEditorOnlyTypes.h b/Engine/Plugins/Experimental/UAF/UAFStateTree/Source/UAFStateTree/Internal/AnimNextStateTreeEditorOnlyTypes.h
new file mode 100644
--- /dev/null
+++ b/Engine/Plugins/Experimental/UAF/UAFStateTree/Source/UAFStateTree/Internal/AnimNextStateTreeEditorOnlyTypes.h
@@ -0,0 +1,35 @@
+// Copyright Epic Games, Inc. All Rights Reserved.
+
+#pragma once
+
+#if WITH_EDITOR
+
+#include "CoreMinimal.h"
+#include "Compilation/AnimNextGetFunctionHeaderCompileContext.h"
+
+class UAnimNextSharedVariables_EditorData;
+class URigVMGraph;
+
+struct FRigVMCompileSettings;
+struct FRigVMClient;
+
+/**
+ * Helper struct to contain info needed to compile RigVM Variables
+ */
+struct UAFSTATETREE_API FAnimNextStateTreeProgrammaticFunctionHeaderParams
+{
+	FAnimNextStateTreeProgrammaticFunctionHeaderParams(UAnimNextSharedVariables_EditorData* InEditorData, const FRigVMCompileSettings& InSettings, FRigVMClient& InRigVMClient, FAnimNextGetFunctionHeaderCompileContext& InOutCompileContext)
+		: EditorData(InEditorData)
+		, Settings(InSettings)
+		, RigVMClient(InRigVMClient)
+		, OutCompileContext(InOutCompileContext)
+	{
+	}
+
+	UAnimNextSharedVariables_EditorData* EditorData;
+	const FRigVMCompileSettings& Settings;
+	FRigVMClient& RigVMClient;
+	FAnimNextGetFunctionHeaderCompileContext& OutCompileContext;
+};
+
+#endif // WITH_EDITOR
\ No newline at end of file
diff --git a/Engine/Plugins/Experimental/UAF/UAFStateTree/Source/UAFStateTree/Private/AnimNextStateTreeContext.cpp b/Engine/Plugins/Experimental/UAF/UAFStateTree/Source/UAFStateTree/Private/AnimNextStateTreeContext.cpp
new file mode 100644
--- /dev/null
+++ b/Engine/Plugins/Experimental/UAF/UAFStateTree/Source/UAFStateTree/Private/AnimNextStateTreeContext.cpp
@@ -0,0 +1,52 @@
+// Copyright Epic Games, Inc. All Rights Reserved.
+
+#include "AnimNextStateTreeContext.h"
+#include "Factory/AnimGraphFactory.h"
+#include "TraitCore/NodeInstance.h"
+#include "TraitCore/TraitStackBinding.h"
+#include "TraitInterfaces/IBlendStack.h"
+#include "TraitInterfaces/ITimeline.h"
+
+bool FAnimNextStateTreeTraitContext::PushAssetOntoBlendStack(TNonNullPtr<UObject> InAsset, const FAlphaBlendArgs& InBlendArguments, const FAnimNextFactoryParams& InFactoryParams) const
+{
+	using namespace UE::UAF;
+
+	TTraitBinding<IBlendStack> BlendStackBinding;
+	if(!Binding->GetStackInterface<IBlendStack>(BlendStackBinding))
+	{
+		return false;
+	}
+
+	// Instantiate/get a graph from the params
+	const UAnimNextAnimationGraph* AnimationGraph = FAnimGraphFactory::GetOrBuildGraph(InAsset, InFactoryParams.GetBuilder());
+	if(AnimationGraph == nullptr)
+	{
+		return false;
+	}
+
+	IBlendStack::FGraphRequest NewGraphRequest;
+	NewGraphRequest.BlendArgs = InBlendArguments;
+	NewGraphRequest.FactoryObject = InAsset;
+	NewGraphRequest.AnimationGraph = AnimationGraph;
+	NewGraphRequest.FactoryParams = InFactoryParams;
+
+	BlendStackBinding.PushGraph(*Context, MoveTemp(NewGraphRequest));
+
+	return true;
+}
+
+float FAnimNextStateTreeTraitContext::QueryPlaybackRatio(TNonNullPtr<UObject> InAsset) const
+{
+	UE::UAF::TTraitBinding<UE::UAF::ITimeline> Timeline;
+	if (Binding->GetStackInterface<UE::UAF::ITimeline>(Timeline))
+	{
+		return Timeline.GetState(*Context).GetPositionRatio();
+	}
+
+	return 1.0f;
+}
+
+FAnimNextGraphInstance& FAnimNextStateTreeTraitContext::GetGraphInstance() const
+{
+	return Binding->GetTraitPtr().GetNodeInstance()->GetOwner();
+}
\ No newline at end of file
diff --git a/Engine/Plugins/Experimental/UAF/UAFStateTree/Source/UAFStateTree/Private/AnimNextStateTreeContext.h b/Engine/Plugins/Experimental/UAF/UAFStateTree/Source/UAFStateTree/Private/AnimNextStateTreeContext.h
new file mode 100644
--- /dev/null
+++ b/Engine/Plugins/Experimental/UAF/UAFStateTree/Source/UAFStateTree/Private/AnimNextStateTreeContext.h
@@ -0,0 +1,56 @@
+// Copyright Epic Games, Inc. All Rights Reserved.
+
+#pragma once
+
+#include "AnimStateTreeTrait.h"
+#include "StructUtils/InstancedStruct.h"
+
+#include "AnimNextStateTreeContext.generated.h"
+
+namespace UE::UAF
+{
+struct FTraitBinding;
+struct FExecutionContext;
+struct FStateTreeTrait;
+}
+
+class UAnimNextAnimationGraph;
+struct FAlphaBlendArgs;
+struct FAnimNextFactoryParams;
+
+USTRUCT()
+struct FAnimNextStateTreeTraitContext
+{	
+	GENERATED_BODY()
+	
+	friend UE::UAF::FStateTreeTrait;
+	friend UE::UAF::FStateTreeTrait::FInstanceData;
+
+	FAnimNextStateTreeTraitContext() = default;
+
+	FAnimNextStateTreeTraitContext(UE::UAF::FExecutionContext& InContext, const UE::UAF::FTraitBinding& InBinding)
+		: Context(&InContext)
+		, Binding(&InBinding)
+	{}
+
+	bool PushAssetOntoBlendStack(TNonNullPtr<UObject> InAsset, const FAlphaBlendArgs& InBlendArguments, const FAnimNextFactoryParams& InFactoryParams) const;
+	float QueryPlaybackRatio(TNonNullPtr<UObject> InAsset) const;
+	FAnimNextGraphInstance& GetGraphInstance() const;
+
+	UE::UAF::FExecutionContext& GetContext() const
+	{
+		check(Context != nullptr);
+		return *Context;
+	}
+
+	const UE::UAF::FTraitBinding& GetBinding() const
+	{
+		check(Binding != nullptr);
+		return *Binding;
+	}
+
+private:
+	UE::UAF::FExecutionContext* Context = nullptr;
+	const UE::UAF::FTraitBinding* Binding = nullptr;
+};
+
diff --git a/Engine/Plugins/Experimental/UAF/UAFStateTree/Source/UAFStateTree/Private/Conditions/AnimNextStateTreeRigVMConditionBase.cpp b/Engine/Plugins/Experimental/UAF/UAFStateTree/Source/UAFStateTree/Private/Conditions/AnimNextStateTreeRigVMConditionBase.cpp
new file mode 100644
--- /dev/null
+++ b/Engine/Plugins/Experimental/UAF/UAFStateTree/Source/UAFStateTree/Private/Conditions/AnimNextStateTreeRigVMConditionBase.cpp
@@ -0,0 +1,162 @@
+// Copyright Epic Games, Inc. All Rights Reserved.
+
+#include "Conditions/AnimNextStateTreeRigVMConditionBase.h"
+#include "AnimNextExecuteContext.h"
+#include "Graph/AnimNextGraphContextData.h"
+#include "Graph/AnimNextGraphInstance.h"
+#include "StateTreeLinker.h"
+#include "StateTreeExecutionContext.h"
+#include "StructUtils/PropertyBag.h"
+#include "TraitCore/ExecutionContext.h"
+#include "AnimNextStateTreeContext.h"
+
+#if WITH_EDITOR
+#include "AnimNextStateTreeEditorOnlyTypes.h"
+#include "RigVMBlueprintGeneratedClass.h"
+#include "RigVMFunctions/RigVMFunctionDefines.h"
+#include "StateTree.h"
+#include "StateTreeEditorData.h"
+#include "UncookedOnlyUtils.h"
+#endif // WITH_EDITOR
+
+#define LOCTEXT_NAMESPACE "AnimNextStateTreeRigVMConditionBase"
+
+#include "UAFRigVMComponent.h"
+
+#include UE_INLINE_GENERATED_CPP_BY_NAME(AnimNextStateTreeRigVMConditionBase)
+
+////////////////////////////////////////////////////////////////////////////////////////////////////////////////////////
+// FAnimNextStateTreeRigVMConditionBase
+
+bool FAnimNextStateTreeRigVMConditionBase::Link(FStateTreeLinker& Linker)
+{
+	Linker.LinkExternalData(TraitContextHandle);
+	return true;
+}
+
+bool FAnimNextStateTreeRigVMConditionBase::TestCondition(FStateTreeExecutionContext& Context) const
+{
+	if (FAnimNextStateTreeRigVMConditionInstanceData* InstanceData = Context.GetInstanceDataPtr<FAnimNextStateTreeRigVMConditionInstanceData>(*this))
+	{
+		FAnimNextStateTreeTraitContext& ExecContext = Context.GetExternalData(TraitContextHandle);
+		FAnimNextGraphInstance& GraphInstance = ExecContext.GetGraphInstance();
+
+		// Need const cast as VM execution isn't const.
+		if (UAnimNextRigVMAsset* RigVMInstance = const_cast<UAnimNextRigVMAsset*>(GraphInstance.GetAsset<UAnimNextRigVMAsset>()))
+		{
+#if WITH_EDITOR
+			if (!InstanceData->FunctionHandle.IsValidForVM(RigVMInstance->GetRigVM()->GetVMHash()))
+#else
+			if (!InstanceData->FunctionHandle.IsValid())
+#endif
+			{
+				InstanceData->FunctionHandle = RigVMInstance->GetFunctionHandle(InternalEventName);
+			}
+
+			if (!ensure(InstanceData->FunctionHandle.IsValid() && ResultIndex != INDEX_NONE))
+			{
+				return false;
+			}
+
+			FUAFRigVMComponent& RigVMComponent = GraphInstance.GetComponent<FUAFRigVMComponent>();
+			FRigVMExtendedExecuteContext& ExtendedExecuteContext = RigVMComponent.GetExtendedExecuteContext();
+			FAnimNextExecuteContext& AnimNextContext = ExtendedExecuteContext.GetPublicDataSafe<FAnimNextExecuteContext>();
+
+			FAnimNextGraphContextData ContextData(GraphInstance.GetModuleInstance(), &GraphInstance);
+			UE::UAF::FScopedExecuteContextData ContextDataScope(AnimNextContext, ContextData);
+
+			RigVMInstance->CallFunctionHandle(InstanceData->FunctionHandle, ExtendedExecuteContext, InstanceData->Parameters);
+
+			// Pull the result out of the property bag
+			return InstanceData->Parameters.GetValueBool(InstanceData->Parameters.GetPropertyBagStruct()->GetPropertyDescs()[ResultIndex]).GetValue();
+		}
+	}
+
+	return true;
+}
+
+#if WITH_EDITOR
+FText FAnimNextStateTreeRigVMConditionBase::GetDescription(const FGuid& ID, FStateTreeDataView InstanceDataView, const IStateTreeBindingLookup& BindingLookup, EStateTreeNodeFormatting Formatting) const
+{
+	return LOCTEXT("AnimNextStateTreeConditon_Desc", "RigVM function driven condition");
+}
+
+void FAnimNextStateTreeRigVMConditionBase::PostEditNodeChangeChainProperty(const FPropertyChangedChainEvent& PropertyChangedEvent, FStateTreeDataView InstanceDataView)
+{
+	if (PropertyChangedEvent.GetMemberPropertyName() == GET_MEMBER_NAME_CHECKED(FAnimNextStateTreeRigVMConditionBase, ConditionFunctionName))
+	{
+		// Function Name selected has changed. Update param struct & result / event names used during execution.
+		if (FAnimNextStateTreeRigVMConditionInstanceData* InstanceData = InstanceDataView.GetMutablePtr<FAnimNextStateTreeRigVMConditionInstanceData>())
+		{
+			// @TODO: This relies on the function name being unique (Ex: In a workspace). For now that's okay. Later on however we will want to use a more robust function picker
+			auto GetRigVMFunctionHeader = [](FName InName)
+			{
+				FRigVMGraphFunctionHeader Result = FRigVMGraphFunctionHeader();
+
+				TMap<FAssetData, FRigVMGraphFunctionHeaderArray> FunctionExports;
+				UE::UAF::UncookedOnly::FUtils::GetExportedFunctionsFromAssetRegistry(UE::UAF::AnimNextPublicGraphFunctionsExportsRegistryTag, FunctionExports);
+				UE::UAF::UncookedOnly::FUtils::GetExportedFunctionsFromAssetRegistry(UE::UAF::ControlRigAssetPublicGraphFunctionsExportsRegistryTag, FunctionExports);
+
+				for (const auto& Export : FunctionExports.Array())
+				{
+					for (const FRigVMGraphFunctionHeader& FunctionHeader : Export.Value.Headers)
+					{
+						if (InName == FunctionHeader.Name)
+						{
+							Result = FunctionHeader;
+							return Result;
+						}
+					}
+				}
+
+				return Result;
+			};
+
+			RigVMFunctionHeader = GetRigVMFunctionHeader(ConditionFunctionName);
+			InstanceData->Parameters.Reset();
+			ResultIndex = INDEX_NONE;
+
+			TArray<FPropertyBagPropertyDesc> PropertyDescs;
+			PropertyDescs.Reserve(RigVMFunctionHeader.Arguments.Num());
+			for (const FRigVMGraphFunctionArgument& Argument : RigVMFunctionHeader.Arguments)
+			{
+				if (Argument.Direction == ERigVMPinDirection::Input || Argument.Direction == ERigVMPinDirection::Output)
+				{
+					FString CppTypeString = Argument.CPPType.ToString();
+					UObject* CppTypeObject = Argument.CPPTypeObject.Get();
+					FPropertyBagPropertyDesc PropertyDesc = FRigVMMemoryStorageStruct::GeneratePropertyBagDescriptor(FRigVMPropertyDescription(Argument.Name, CppTypeString, CppTypeObject, Argument.DefaultValue));
+					if (Argument.Direction == ERigVMPinDirection::Output && Argument.CPPType == RigVMTypeUtils::BoolTypeName)
+					{
+						PropertyDesc.PropertyFlags &= ~CPF_Edit;	// Hide the result from the UI
+						ResultIndex = PropertyDescs.Num();
+					}
+					PropertyDescs.Add(PropertyDesc);
+				}
+			}
+
+			// TODO: Constrain functions we can pick to those returning only a bool
+			ensure(ResultIndex != INDEX_NONE);
+
+			InstanceData->Parameters.AddProperties(PropertyDescs);
+		}
+	}
+}
+
+void FAnimNextStateTreeRigVMConditionBase::GetProgrammaticFunctionHeaders(FAnimNextStateTreeProgrammaticFunctionHeaderParams& InProgrammaticFunctionHeaderParams, const UStateTreeState* State, const FStateTreeBindableStructDesc& Desc)
+{
+	using namespace UE::UAF::UncookedOnly;
+
+	FAnimNextGetFunctionHeaderCompileContext& OutCompileContext = InProgrammaticFunctionHeaderParams.OutCompileContext;
+	
+	StateName = State->Name;
+	NodeId = Desc.ID;
+	InternalEventName = FName(FUtils::MakeFunctionWrapperEventName(RigVMFunctionHeader.Name));
+
+	FAnimNextProgrammaticFunctionHeader AnimNextFunctionHeader = {};
+	AnimNextFunctionHeader.FunctionHeader = RigVMFunctionHeader;
+	OutCompileContext.GetMutableFunctionHeaders().Add(AnimNextFunctionHeader);
+}
+
+#endif // WITH_EDITOR
+
+#undef LOCTEXT_NAMESPACE
diff --git a/Engine/Plugins/Experimental/UAF/UAFStateTree/Source/UAFStateTree/Private/Tasks/AnimNextStateTreeGraphInstanceTask.cpp b/Engine/Plugins/Experimental/UAF/UAFStateTree/Source/UAFStateTree/Private/Tasks/AnimNextStateTreeGraphInstanceTask.cpp
new file mode 100644
--- /dev/null
+++ b/Engine/Plugins/Experimental/UAF/UAFStateTree/Source/UAFStateTree/Private/Tasks/AnimNextStateTreeGraphInstanceTask.cpp
@@ -0,0 +1,161 @@
+// Copyright Epic Games, Inc. All Rights Reserved.
+
+#include "Tasks/AnimNextStateTreeGraphInstanceTask.h"
+
+#include "AnimNextAnimGraphSettings.h"
+#include "StateTreeLinker.h"
+#include "StateTreeExecutionContext.h"
+#include "Factory/AnimNextFactoryParams.h"
+#include "AnimNextStateTreeContext.h"
+#include "Factory/AnimGraphFactory.h"
+#include "UObject/FortniteMainBranchObjectVersion.h"
+
+#if WITH_EDITORONLY_DATA
+bool FAnimNextGraphInstanceTaskInstanceData::Serialize(FArchive& Ar)
+{
+	Ar.UsingCustomVersion(FFortniteMainBranchObjectVersion::GUID);
+	return false;
+}
+
+void FAnimNextGraphInstanceTaskInstanceData::PostSerialize(const FArchive& Ar)
+{
+	if (Ar.IsLoading() && Ar.CustomVer(FFortniteMainBranchObjectVersion::GUID) < FFortniteMainBranchObjectVersion::AnimNextVariableReferences)
+	{
+		// Prior to this we only supported 'imprecise' payloads so try to do our best be existing content, but some data loss is possible
+		if (Asset != nullptr && !Asset->IsA<UAnimNextAnimationGraph>())
+		{
+			Parameters = UE::UAF::FAnimGraphFactory::GetDefaultParamsForObject(Asset);
+
+			// Match any legacy parameter values by name
+			if(Payload_DEPRECATED.GetNumPropertiesInBag() > 0 && Parameters.Builder.Stacks.Num() > 0)
+			{
+				const UPropertyBag* PropertyBag = Payload_DEPRECATED.GetPropertyBagStruct();
+				const uint8* SourceContainerPtr = Payload_DEPRECATED.GetValue().GetMemory();
+
+				auto PatchInLegacyPropertyByName = [this, SourceContainerPtr](const FPropertyBagPropertyDesc& InDesc)
+				{
+					for (TInstancedStruct<FAnimNextTraitSharedData>& Struct : Parameters.Builder.Stacks[0].TraitStructs)
+					{
+						uint8* TargetContainerPtr = Struct.GetMutableMemory();
+						for (TFieldIterator<FProperty> It(Struct.GetScriptStruct()); It; ++It)
+						{
+							auto NameMatches = [](const FProperty* InProperty, FName InNameA, FName InNameB)
+							{
+								if (InNameA == InNameB)
+								{
+									return true;
+								}
+
+								// Match bool properties with no 'b' prefix
+								if (InProperty->IsA<FBoolProperty>())
+								{
+									FString StringA = InNameA.ToString();
+									StringA.RemoveFromStart(TEXT("b"));
+									FString StringB = InNameB.ToString();
+									StringB.RemoveFromStart(TEXT("b"));
+									return StringA.Equals(StringB, ESearchCase::IgnoreCase);
+								}
+								return false; 
+							};
+
+							if (NameMatches(*It, It->GetFName(), InDesc.Name) && It->GetClass() == InDesc.CachedProperty->GetClass())
+							{
+								const uint8* SourcePtr = InDesc.CachedProperty->ContainerPtrToValuePtr<uint8>(SourceContainerPtr);
+								uint8* TargetPtr = It->ContainerPtrToValuePtr<uint8>(TargetContainerPtr);
+								It->CopyCompleteValue(TargetPtr, SourcePtr);
+								break;
+							}
+						}
+					}
+				};
+
+				for (const FPropertyBagPropertyDesc& Desc : PropertyBag->GetPropertyDescs())
+				{
+					PatchInLegacyPropertyByName(Desc);
+				}
+			}
+		}
+	}
+}
+#endif
+
+FAnimNextStateTreeGraphInstanceTask::FAnimNextStateTreeGraphInstanceTask()
+{
+	// Re-selecting the same state should not cause a re-trigger of EnterState()
+	bShouldStateChangeOnReselect = false;
+}
+
+bool FAnimNextStateTreeGraphInstanceTask::Link(FStateTreeLinker& Linker)
+{
+	Linker.LinkExternalData(TraitContextHandle);	
+	return true;
+}
+
+EStateTreeRunStatus FAnimNextStateTreeGraphInstanceTask::EnterState(FStateTreeExecutionContext& Context, const FStateTreeTransitionResult& Transition) const
+{
+	FAnimNextStateTreeTraitContext& ExecContext = Context.GetExternalData(TraitContextHandle);
+	FInstanceDataType& InstanceData = Context.GetInstanceData(*this);
+
+	if (ExecContext.PushAssetOntoBlendStack(InstanceData.Asset.Get(), InstanceData.BlendOptions, InstanceData.Parameters))
+	{
+		return EStateTreeRunStatus::Running;
+	}
+	
+	return EStateTreeRunStatus::Failed;
+}
+
+EStateTreeRunStatus FAnimNextStateTreeGraphInstanceTask::Tick(FStateTreeExecutionContext& Context, const float DeltaTime) const
+{
+	FInstanceDataType& InstanceData = Context.GetInstanceData(*this);
+	if (InstanceData.bContinueTicking)
+	{
+		FAnimNextStateTreeTraitContext& ExecContext = Context.GetExternalData(TraitContextHandle);
+
+		InstanceData.PlaybackRatio = ExecContext.QueryPlaybackRatio(InstanceData.Asset.Get());
+
+		return EStateTreeRunStatus::Running;
+	}
+	
+	return EStateTreeRunStatus::Succeeded;
+}
+
+void FAnimNextStateTreeGraphInstanceTask::ExitState(FStateTreeExecutionContext& Context, const FStateTreeTransitionResult& Transition) const
+{
+	FAnimNextStateTreeTaskBase::ExitState(Context, Transition);
+}
+
+#if WITH_EDITOR
+void FAnimNextStateTreeGraphInstanceTask::PostEditInstanceDataChangeChainProperty(const FPropertyChangedChainEvent& PropertyChangedEvent, FStateTreeDataView InstanceDataView)
+{
+	const FProperty* Property = PropertyChangedEvent.Property;
+	if(Property == nullptr)
+	{
+		return;
+	}
+
+	if(Property->GetFName() == GET_MEMBER_NAME_CHECKED(FAnimNextGraphInstanceTaskInstanceData, Asset))
+	{
+		// Repopulate payload property bag if the selected asset changes
+		FAnimNextGraphInstanceTaskInstanceData& InstanceData = InstanceDataView.GetMutable<FAnimNextGraphInstanceTaskInstanceData>();
+
+		if (InstanceData.Asset == nullptr || InstanceData.Asset.IsA<UAnimNextAnimationGraph>())
+		{
+			InstanceData.Parameters.Reset();
+		}
+		else
+		{
+			// TODO: Copy over matching values before overwrite to avoid losing user edits to matching trait data
+			InstanceData.Parameters = UE::UAF::FAnimGraphFactory::GetDefaultParamsForObject(InstanceData.Asset);
+		}
+	}
+}
+
+void FAnimNextStateTreeGraphInstanceTask::GetObjectReferences(TArray<const UObject*>& OutReferencedObjects, FStateTreeDataView InstanceDataView) const
+{
+	if (InstanceDataView.IsValid())
+	{
+		OutReferencedObjects.Add(InstanceDataView.Get<FAnimNextGraphInstanceTaskInstanceData>().Asset.Get());
+	}	
+}
+#endif
+
diff --git a/Engine/Plugins/Experimental/UAF/UAFStateTree/Source/UAFStateTree/Private/Tasks/AnimNextStateTreeRigVMTaskBase.cpp b/Engine/Plugins/Experimental/UAF/UAFStateTree/Source/UAFStateTree/Private/Tasks/AnimNextStateTreeRigVMTaskBase.cpp
new file mode 100644
--- /dev/null
+++ b/Engine/Plugins/Experimental/UAF/UAFStateTree/Source/UAFStateTree/Private/Tasks/AnimNextStateTreeRigVMTaskBase.cpp
@@ -0,0 +1,164 @@
+// Copyright Epic Games, Inc. All Rights Reserved.
+
+#include "Tasks/AnimNextStateTreeRigVMTaskBase.h"
+#include "AnimNextExecuteContext.h"
+#include "Graph/AnimNextGraphContextData.h"
+#include "Graph/AnimNextGraphInstance.h"
+#include "StateTreeLinker.h"
+#include "StateTreeExecutionContext.h"
+#include "StructUtils/PropertyBag.h"
+#include "TraitCore/ExecutionContext.h"
+#include "AnimNextStateTreeContext.h"
+
+#if WITH_EDITOR
+#include "AnimNextStateTreeEditorOnlyTypes.h"
+#include "RigVMBlueprintGeneratedClass.h"
+#include "RigVMFunctions/RigVMFunctionDefines.h"
+#include "StateTree.h"
+#include "StateTreeEditorData.h"
+#include "UncookedOnlyUtils.h"
+#endif // WITH_EDITOR
+
+#define LOCTEXT_NAMESPACE "AnimNextStateTreeRigVMTaskBase"
+
+#include "AnimNextRigVMAsset.h"
+#include "UAFRigVMComponent.h"
+
+#include UE_INLINE_GENERATED_CPP_BY_NAME(AnimNextStateTreeRigVMTaskBase)
+
+////////////////////////////////////////////////////////////////////////////////////////////////////////////////////////
+// FAnimNextStateTreeRigVMTaskBase
+
+bool FAnimNextStateTreeRigVMTaskBase::Link(FStateTreeLinker& Linker)
+{
+	Linker.LinkExternalData(TraitContextHandle);
+	return true;
+}
+
+EStateTreeRunStatus FAnimNextStateTreeRigVMTaskBase::EnterState(FStateTreeExecutionContext& Context, const FStateTreeTransitionResult& Transition) const
+{
+	if (FAnimNextStateTreeRigVMTaskInstanceData* InstanceData = Context.GetInstanceDataPtr<FAnimNextStateTreeRigVMTaskInstanceData>(*this))
+	{
+		FAnimNextStateTreeTraitContext& ExecContext = Context.GetExternalData(TraitContextHandle);
+		FAnimNextGraphInstance& GraphInstance = ExecContext.GetGraphInstance();
+
+		// Need const cast as VM execution isn't const.
+		if (UAnimNextRigVMAsset* RigVMAsset = const_cast<UAnimNextRigVMAsset*>(GraphInstance.GetAsset<UAnimNextRigVMAsset>()))
+		{
+#if WITH_EDITOR
+			if (!InstanceData->FunctionHandle.IsValidForVM(RigVMAsset->GetRigVM()->GetVMHash()))
+#else
+			if (!InstanceData->FunctionHandle.IsValid())
+#endif
+			{
+				InstanceData->FunctionHandle = RigVMAsset->GetFunctionHandle(InternalEventName);
+			}
+
+			if (!ensure(InstanceData->FunctionHandle.IsValid() && ResultIndex != INDEX_NONE))
+			{
+				return EStateTreeRunStatus::Failed;
+			}
+
+			FUAFRigVMComponent& RigVMComponent = GraphInstance.GetComponent<FUAFRigVMComponent>();
+			FRigVMExtendedExecuteContext& ExtendedExecuteContext = RigVMComponent.GetExtendedExecuteContext();
+			FAnimNextExecuteContext& AnimNextContext = ExtendedExecuteContext.GetPublicDataSafe<FAnimNextExecuteContext>();
+
+			FAnimNextGraphContextData ContextData(GraphInstance.GetModuleInstance(), &GraphInstance);
+			UE::UAF::FScopedExecuteContextData ContextDataScope(AnimNextContext, ContextData);
+
+			// @TODO: Add option to execute per tick?
+			RigVMAsset->CallFunctionHandle(InstanceData->FunctionHandle, ExtendedExecuteContext, InstanceData->Parameters);
+
+			// @TODO: For now we use a fire & forget model. But in the future we might want users to be able to bind to the result value from a function
+			// This will require a special UI customization, & standalone execute is useful for modifying existing vars. So save for later effort.	
+			//return InstanceData->Parameters.GetValueBool(InstanceData->Parameters.GetPropertyBagStruct()->GetPropertyDescs()[ResultIndex]).GetValue();
+
+			return EStateTreeRunStatus::Running;
+		}
+	}
+
+	return EStateTreeRunStatus::Failed;
+}
+
+#if WITH_EDITOR
+FText FAnimNextStateTreeRigVMTaskBase::GetDescription(const FGuid& ID, FStateTreeDataView InstanceDataView, const IStateTreeBindingLookup& BindingLookup, EStateTreeNodeFormatting Formatting) const
+{
+	return LOCTEXT("AnimNextStateTreeConditon_Desc", "RigVM function driven Task");
+}
+
+void FAnimNextStateTreeRigVMTaskBase::PostEditNodeChangeChainProperty(const FPropertyChangedChainEvent& PropertyChangedEvent, FStateTreeDataView InstanceDataView)
+{
+	if (PropertyChangedEvent.GetMemberPropertyName() == GET_MEMBER_NAME_CHECKED(FAnimNextStateTreeRigVMTaskBase, TaskFunctionName))
+	{
+		// Function Name selected has changed. Update param struct & result / event names used during execution.
+		if (FAnimNextStateTreeRigVMTaskInstanceData* InstanceData = InstanceDataView.GetMutablePtr<FAnimNextStateTreeRigVMTaskInstanceData>())
+		{
+			// @TODO: This relies on the function name being unique (Ex: In a workspace). For now that's okay. Later on however we will want to use a more robust function picker
+			auto GetRigVMFunctionHeader = [](FName InName)
+			{
+				FRigVMGraphFunctionHeader Result = FRigVMGraphFunctionHeader();
+
+				TMap<FAssetData, FRigVMGraphFunctionHeaderArray> FunctionExports;
+				UE::UAF::UncookedOnly::FUtils::GetExportedFunctionsFromAssetRegistry(UE::UAF::AnimNextPublicGraphFunctionsExportsRegistryTag, FunctionExports);
+				UE::UAF::UncookedOnly::FUtils::GetExportedFunctionsFromAssetRegistry(UE::UAF::ControlRigAssetPublicGraphFunctionsExportsRegistryTag, FunctionExports);
+
+				for (const auto& Export : FunctionExports.Array())
+				{
+					for (const FRigVMGraphFunctionHeader& FunctionHeader : Export.Value.Headers)
+					{
+						if (InName == FunctionHeader.Name)
+						{
+							Result = FunctionHeader;
+							return Result;
+						}
+					}
+				}
+
+				return Result;
+			};
+
+			RigVMFunctionHeader = GetRigVMFunctionHeader(TaskFunctionName);
+			InstanceData->Parameters.Reset();
+			ResultIndex = INDEX_NONE;
+
+			TArray<FPropertyBagPropertyDesc> PropertyDescs;
+			PropertyDescs.Reserve(RigVMFunctionHeader.Arguments.Num());
+			for (const FRigVMGraphFunctionArgument& Argument : RigVMFunctionHeader.Arguments)
+			{
+				if (Argument.Direction == ERigVMPinDirection::Input || Argument.Direction == ERigVMPinDirection::Output)
+				{
+					FString CppTypeString = Argument.CPPType.ToString();
+					UObject* CppTypeObject = Argument.CPPTypeObject.Get();
+					FPropertyBagPropertyDesc PropertyDesc = FRigVMMemoryStorageStruct::GeneratePropertyBagDescriptor(FRigVMPropertyDescription(Argument.Name, CppTypeString, CppTypeObject, Argument.DefaultValue));
+					if (Argument.Direction == ERigVMPinDirection::Output && Argument.CPPType == RigVMTypeUtils::BoolTypeName)
+					{
+						PropertyDesc.PropertyFlags &= ~CPF_Edit;	// Hide the result from the UI
+						ResultIndex = PropertyDescs.Num();
+					}
+					PropertyDescs.Add(PropertyDesc);
+				}
+			}
+		
+			InstanceData->Parameters.AddProperties(PropertyDescs);
+		}
+	}
+}
+
+void FAnimNextStateTreeRigVMTaskBase::GetProgrammaticFunctionHeaders(FAnimNextStateTreeProgrammaticFunctionHeaderParams& InProgrammaticFunctionHeaderParams, const UStateTreeState* State, const FStateTreeBindableStructDesc& Desc)
+{
+	using namespace UE::UAF::UncookedOnly;
+
+	FAnimNextGetFunctionHeaderCompileContext& OutCompileContext = InProgrammaticFunctionHeaderParams.OutCompileContext;
+	
+	StateName = State->Name;
+	NodeId = Desc.ID;
+	InternalEventName = FName(FUtils::MakeFunctionWrapperEventName(RigVMFunctionHeader.Name));
+
+	FAnimNextProgrammaticFunctionHeader AnimNextFunctionHeader = {};
+	AnimNextFunctionHeader.FunctionHeader = RigVMFunctionHeader;
+	OutCompileContext.GetMutableFunctionHeaders().Add(AnimNextFunctionHeader);
+}
+
+#endif // WITH_EDITOR
+
+#undef LOCTEXT_NAMESPACE
diff --git a/Engine/Plugins/Experimental/UAF/UAFStateTree/Source/UAFStateTreeUncookedOnly/Internal/AnimNextStateTree_EditorData.h b/Engine/Plugins/Experimental/UAF/UAFStateTree/Source/UAFStateTreeUncookedOnly/Internal/AnimNextStateTree_EditorData.h
new file mode 100644
--- /dev/null
+++ b/Engine/Plugins/Experimental/UAF/UAFStateTree/Source/UAFStateTreeUncookedOnly/Internal/AnimNextStateTree_EditorData.h
@@ -0,0 +1,33 @@
+// Copyright Epic Games, Inc. All Rights Reserved.
+
+#pragma once
+
+#include "Graph/AnimNextAnimationGraph_EditorData.h"
+
+#include "AnimNextStateTree_EditorData.generated.h"
+
+UCLASS()
+class UAFSTATETREEUNCOOKEDONLY_API UAnimNextStateTree_EditorData : public UAnimNextAnimationGraph_EditorData
+{
+	GENERATED_BODY()
+	
+protected:
+	virtual TSubclassOf<UAssetUserData> GetAssetUserDataClass() const override;
+
+	friend class UAnimNextStateTreeFactory;
+	friend class UAnimNextStateTreeTreeEditorData;
+
+	// IRigVMClientHost interface
+	virtual void RecompileVM() override;
+	
+	// UAnimNextRigVMAssetEditorData interface
+	virtual TConstArrayView<TSubclassOf<UAnimNextRigVMAssetEntry>> GetEntryClasses() const override;
+	virtual void OnPreCompileGetProgrammaticFunctionHeaders(const FRigVMCompileSettings& InSettings, FAnimNextGetFunctionHeaderCompileContext& OutCompileContext) override;
+	virtual void OnPreCompileGetProgrammaticGraphs(const FRigVMCompileSettings& InSettings, FAnimNextGetGraphCompileContext& OutCompileContext) override;
+	virtual void OnPostCompileVariables(const FRigVMCompileSettings& InSettings, const FAnimNextGetVariableCompileContext& InCompileContext) override;
+
+private:
+	// Property bag that forms the root parameters. Consists of all combined variables from shared variable assets/structs.
+	UPROPERTY()
+	FInstancedPropertyBag CombinedVariablesPropertyBag;
+};
\ No newline at end of file
diff --git a/Engine/Plugins/Experimental/UAF/UAFWarping/Config/BaseUAFWarping.ini b/Engine/Plugins/Experimental/UAF/UAFWarping/Config/BaseUAFWarping.ini
new file mode 100644
--- /dev/null
+++ b/Engine/Plugins/Experimental/UAF/UAFWarping/Config/BaseUAFWarping.ini
@@ -0,0 +1,11 @@
+﻿[CoreRedirects]
+; AnimNextWarping plugin -> UAFWarping plugin
++PackageRedirects=(OldName="/Script/AnimNextWarping", NewName="/Script/UAFWarping")
+
++StructRedirects=(OldName="/Script/AnimNextWarping.SteeringTraitSharedData", NewName="/Script/UAFWarping.SteeringTraitSharedData")
++StructRedirects=(OldName="/Script/AnimNextWarping.AnimNextSteeringTask", NewName="/Script/UAFWarping.AnimNextSteeringTask")
++StructRedirects=(OldName="/Script/AnimNextWarping.StrafeWarpFootData", NewName="/Script/UAFWarping.StrafeWarpFootData")
++StructRedirects=(OldName="/Script/AnimNextWarping.StrafeWarpingTraitSharedData", NewName="/Script/UAFWarping.StrafeWarpingTraitSharedData")
++StructRedirects=(OldName="/Script/AnimNextWarping.AnimNextStrafeWarpingTask", NewName="/Script/UAFWarping.AnimNextStrafeWarpingTask")
++StructRedirects=(OldName="/Script/AnimNextWarping.WarpTestTraitSharedData", NewName="/Script/UAFWarping.WarpTestTraitSharedData")
++StructRedirects=(OldName="/Script/AnimNextWarping.AnimNextWarpTestTask", NewName="/Script/UAFWarping.AnimNextWarpTestTask")
\ No newline at end of file
