Integrate UAF StateTree with UAF Graph workspace exports.
- This makes added functions show up in the outliner

diff --git a/Engine/Plugins/Experimental/UAF/UAFAnimGraph/Source/UAFAnimGraphUncookedOnly/Internal/AnimNextAnimGraphWorkspaceAssetUserData.h b/Engine/Plugins/Experimental/UAF/UAFAnimGraph/Source/UAFAnimGraphUncookedOnly/Internal/AnimNextAnimGraphWorkspaceAssetUserData.h
--- a/Engine/Plugins/Experimental/UAF/UAFAnimGraph/Source/UAFAnimGraphUncookedOnly/Internal/AnimNextAnimGraphWorkspaceAssetUserData.h
+++ b/Engine/Plugins/Experimental/UAF/UAFAnimGraph/Source/UAFAnimGraphUncookedOnly/Internal/AnimNextAnimGraphWorkspaceAssetUserData.h
@@ -1,4 +1,4 @@
-﻿// Copyright Epic Games, Inc. All Rights Reserved.
+// Copyright Epic Games, Inc. All Rights Reserved.
 
 #pragma once
 
@@ -8,6 +8,8 @@
 
 #include "AnimNextAnimGraphWorkspaceAssetUserData.generated.h"
 
+#define UE_API UAFANIMGRAPHUNCOOKEDONLY_API
+
 USTRUCT()
 struct FAnimNextAnimationGraphOutlinerData : public FAnimNextRigVMAssetOutlinerData
 {
@@ -21,16 +23,28 @@ struct FAnimNextAnimationGraphOutlinerData : public FAnimNextRigVMAssetOutlinerD
 	}
 };
 
-UCLASS()
+UCLASS(MinimalAPI)
 class UAnimNextAnimGraphWorkspaceAssetUserData : public UAssetUserData
 {
+	GENERATED_BODY()
+
 public:
+
 	virtual bool IsEditorOnly() const override { return true; }
 
-private:
-	GENERATED_BODY()
+protected:
+
+	UE_API virtual void GetAssetRegistryTags(FAssetRegistryTagsContext Context) const override;
+
+	/** Get root asset workspace exports. */
+	UE_API virtual void GetRootAssetExport(FAssetRegistryTagsContext Context) const;
+
+	/** Get all workspace asset exports except for the root asset. */
+	UE_API virtual void GetWorkspaceAssetExports(FAssetRegistryTagsContext Context) const;
+
+protected:
 
-	virtual void GetAssetRegistryTags(FAssetRegistryTagsContext Context) const override;
-	
 	mutable FWorkspaceOutlinerItemExports CachedExports;
 };
+
+#undef UE_API
\ No newline at end of file
diff --git a/Engine/Plugins/Experimental/UAF/UAFAnimGraph/Source/UAFAnimGraphUncookedOnly/Private/AnimNextAnimGraphWorkspaceAssetUserData.cpp b/Engine/Plugins/Experimental/UAF/UAFAnimGraph/Source/UAFAnimGraphUncookedOnly/Private/AnimNextAnimGraphWorkspaceAssetUserData.cpp
--- a/Engine/Plugins/Experimental/UAF/UAFAnimGraph/Source/UAFAnimGraphUncookedOnly/Private/AnimNextAnimGraphWorkspaceAssetUserData.cpp
+++ b/Engine/Plugins/Experimental/UAF/UAFAnimGraph/Source/UAFAnimGraphUncookedOnly/Private/AnimNextAnimGraphWorkspaceAssetUserData.cpp
@@ -15,18 +15,8 @@ void UAnimNextAnimGraphWorkspaceAssetUserData::GetAssetRegistryTags(FAssetRegist
 	if (!Context.IsSaving())
 	{
 		CachedExports.Exports.Reset();
-		{
-			UAnimNextAnimationGraph* Asset = CastChecked<UAnimNextAnimationGraph>(GetOuter());
-			const UAnimNextAnimationGraph_EditorData* EditorData = UE::UAF::UncookedOnly::FUtils::GetEditorData<UAnimNextAnimationGraph_EditorData>(Asset);
-			{
-				FWorkspaceOutlinerItemExport& RootAssetExport = CachedExports.Exports.Add_GetRef(FWorkspaceOutlinerItemExport(Asset->GetFName(), Asset));
-				RootAssetExport.GetData().InitializeAsScriptStruct(FAnimNextAnimationGraphOutlinerData::StaticStruct());
-				FAnimNextRigVMAssetOutlinerData& Data = RootAssetExport.GetData().GetMutable<FAnimNextRigVMAssetOutlinerData>();
-				Data.SoftAssetPtr = Asset;
-			}
-
-			UE::UAF::UncookedOnly::FUtils::GetAssetWorkspaceExports(EditorData, CachedExports, Context);
-		}
+		GetRootAssetExport(Context);
+		GetWorkspaceAssetExports(Context);
 	}
 
 	FString TagValue;
@@ -34,3 +24,23 @@ void UAnimNextAnimGraphWorkspaceAssetUserData::GetAssetRegistryTags(FAssetRegist
 	Context.AddTag(FAssetRegistryTag(UE::Workspace::ExportsWorkspaceItemsRegistryTag, TagValue, FAssetRegistryTag::TT_Hidden));
 }
 
+void UAnimNextAnimGraphWorkspaceAssetUserData::GetRootAssetExport(FAssetRegistryTagsContext Context) const
+{
+	UAnimNextAnimationGraph* Asset = CastChecked<UAnimNextAnimationGraph>(GetOuter());
+	FWorkspaceOutlinerItemExport& RootAssetExport = CachedExports.Exports.Add_GetRef(FWorkspaceOutlinerItemExport(Asset->GetFName(), Asset));
+
+	RootAssetExport.GetData().InitializeAsScriptStruct(FAnimNextAnimationGraphOutlinerData::StaticStruct());
+	FAnimNextRigVMAssetOutlinerData& Data = RootAssetExport.GetData().GetMutable<FAnimNextRigVMAssetOutlinerData>();
+	Data.SoftAssetPtr = Asset;
+}
+
+void UAnimNextAnimGraphWorkspaceAssetUserData::GetWorkspaceAssetExports(FAssetRegistryTagsContext Context) const
+{
+	Super::GetAssetRegistryTags(Context);
+
+	UAnimNextAnimationGraph* Asset = CastChecked<UAnimNextAnimationGraph>(GetOuter());
+	const UAnimNextAnimationGraph_EditorData* EditorData = UE::UAF::UncookedOnly::FUtils::GetEditorData<UAnimNextAnimationGraph_EditorData>(Asset);
+
+	UE::UAF::UncookedOnly::FUtils::GetAssetWorkspaceExports(EditorData, CachedExports, Context);
+}
+
diff --git a/Engine/Plugins/Experimental/UAF/UAFStateTree/Source/UAFStateTreeUncookedOnly/Internal/AnimNextStateTreeWorkspaceAssetUserData.h b/Engine/Plugins/Experimental/UAF/UAFStateTree/Source/UAFStateTreeUncookedOnly/Internal/AnimNextStateTreeWorkspaceAssetUserData.h
--- a/Engine/Plugins/Experimental/UAF/UAFStateTree/Source/UAFStateTreeUncookedOnly/Internal/AnimNextStateTreeWorkspaceAssetUserData.h
+++ b/Engine/Plugins/Experimental/UAF/UAFStateTree/Source/UAFStateTreeUncookedOnly/Internal/AnimNextStateTreeWorkspaceAssetUserData.h
@@ -4,12 +4,18 @@
 
 #include "Engine/AssetUserData.h"
 
+#include "AnimNextAnimGraphWorkspaceAssetUserData.h"
 #include "AnimNextStateTreeWorkspaceAssetUserData.generated.h"
 
 UCLASS()
-class UAnimNextStateTreeWorkspaceAssetUserData : public UAssetUserData
+class UAnimNextStateTreeWorkspaceAssetUserData : public UAnimNextAnimGraphWorkspaceAssetUserData
 {
 	GENERATED_BODY()
 
-	virtual void GetAssetRegistryTags(FAssetRegistryTagsContext Context) const override;
+protected:
+
+	//~ Begin UAnimNextAnimGraphWorkspaceAssetUserData Interface
+	virtual void GetRootAssetExport(FAssetRegistryTagsContext Context) const override;
+	virtual void GetWorkspaceAssetExports(FAssetRegistryTagsContext Context) const override;
+	//~ End UAnimNextAnimGraphWorkspaceAssetUserData Interface
 };
diff --git a/Engine/Plugins/Experimental/UAF/UAFStateTree/Source/UAFStateTreeUncookedOnly/Private/AnimNextStateTreeWorkspaceAssetUserData.cpp b/Engine/Plugins/Experimental/UAF/UAFStateTree/Source/UAFStateTreeUncookedOnly/Private/AnimNextStateTreeWorkspaceAssetUserData.cpp
--- a/Engine/Plugins/Experimental/UAF/UAFStateTree/Source/UAFStateTreeUncookedOnly/Private/AnimNextStateTreeWorkspaceAssetUserData.cpp
+++ b/Engine/Plugins/Experimental/UAF/UAFStateTree/Source/UAFStateTreeUncookedOnly/Private/AnimNextStateTreeWorkspaceAssetUserData.cpp
@@ -9,27 +9,36 @@
 #include "UObject/AssetRegistryTagsContext.h"
 #include "WorkspaceAssetRegistryInfo.h"
 
-void UAnimNextStateTreeWorkspaceAssetUserData::GetAssetRegistryTags(FAssetRegistryTagsContext Context) const
+void UAnimNextStateTreeWorkspaceAssetUserData::GetRootAssetExport(FAssetRegistryTagsContext Context) const
 {
-	FWorkspaceOutlinerItemExports OutlinerExports;
+	if (UAnimNextStateTree* AnimStateTree = CastChecked<UAnimNextStateTree>(GetOuter()))
+	{
+		const FName AnimStateTreeIdentifier = AnimStateTree->GetFName();
+		UStateTree* StateTree = AnimStateTree->StateTree;
+		FWorkspaceOutlinerItemExport& RootAssetExport = CachedExports.Exports.Add_GetRef(FWorkspaceOutlinerItemExport(AnimStateTreeIdentifier, GetOuter()));
+		RootAssetExport.GetData().InitializeAsScriptStruct(FAnimNextStateTreeOutlinerData::StaticStruct());
+		RootAssetExport.GetData().GetMutable<FAnimNextStateTreeOutlinerData>().SoftAssetPtr = AnimStateTree;
+	}
+}
+
+void UAnimNextStateTreeWorkspaceAssetUserData::GetWorkspaceAssetExports(FAssetRegistryTagsContext Context) const
+{
+	Super::GetWorkspaceAssetExports(Context);
 
 	if (UAnimNextStateTree* AnimStateTree = CastChecked<UAnimNextStateTree>(GetOuter()))
 	{		
 		const FName AnimStateTreeIdentifier = AnimStateTree->GetFName();
 		
-		// StateTree asset export
 		UStateTree* StateTree = AnimStateTree->StateTree;
-		FWorkspaceOutlinerItemExport& RootAssetExport = OutlinerExports.Exports.Add_GetRef(FWorkspaceOutlinerItemExport(AnimStateTreeIdentifier, GetOuter()));			
-		RootAssetExport.GetData().InitializeAsScriptStruct(FAnimNextStateTreeOutlinerData::StaticStruct());
-		RootAssetExport.GetData().GetMutable<FAnimNextStateTreeOutlinerData>().SoftAssetPtr = AnimStateTree;
+		FWorkspaceOutlinerItemExport& RootAssetExport = CachedExports.Exports[0];
 
 		if (const UStateTreeEditorData* EditorData = Cast<UStateTreeEditorData>(StateTree->EditorData)) 
 		{
 			// Export each state as individual item as well (parented according to hierarchy)
 			TMap<UStateTreeState*, FWorkspaceOutlinerItemExport> ParentExports;
-			EditorData->VisitHierarchy([&OutlinerExports, RootAssetExport, &ParentExports, EditorData](UStateTreeState& State, const UStateTreeState* ParentState) -> EStateTreeVisitor
+			EditorData->VisitHierarchy([&CachedExports = CachedExports, RootAssetExport, &ParentExports, EditorData](UStateTreeState& State, const UStateTreeState* ParentState) -> EStateTreeVisitor
 			{
-				FWorkspaceOutlinerItemExport& StateExport = OutlinerExports.Exports.Add_GetRef(
+				FWorkspaceOutlinerItemExport& StateExport = CachedExports.Exports.Add_GetRef(
 				ParentState == nullptr ? FWorkspaceOutlinerItemExport(State.Name, RootAssetExport)
 					: FWorkspaceOutlinerItemExport(State.Name, ParentExports.FindChecked(ParentState))
 				);
@@ -68,7 +77,7 @@ void UAnimNextStateTreeWorkspaceAssetUserData::GetAssetRegistryTags(FAssetRegist
 						if (ReferredObject != nullptr)
 						{
 							FSoftObjectPath Path = ReferredObject;
-							FWorkspaceOutlinerItemExport& GraphReference = OutlinerExports.Exports.Add_GetRef(FWorkspaceOutlinerItemExport(FName(Path.ToString()), StateExport));
+							FWorkspaceOutlinerItemExport& GraphReference = CachedExports.Exports.Add_GetRef(FWorkspaceOutlinerItemExport(FName(Path.ToString()), StateExport));
 							GraphReference.GetData().InitializeAsScriptStruct(FWorkspaceOutlinerAssetReferenceItemData::StaticStruct());
 							FWorkspaceOutlinerAssetReferenceItemData& ReferenceData = GraphReference.GetData().GetMutable<FWorkspaceOutlinerAssetReferenceItemData>();
 							ReferenceData.ReferredObjectPath = Path;
@@ -81,11 +90,5 @@ void UAnimNextStateTreeWorkspaceAssetUserData::GetAssetRegistryTags(FAssetRegist
 				return EStateTreeVisitor::Continue;
 			});
 		}
-
-		{
-			FString TagValue;
-			FWorkspaceOutlinerItemExports::StaticStruct()->ExportText(TagValue, &OutlinerExports, nullptr, nullptr, PPF_None, nullptr);
-			Context.AddTag(UObject::FAssetRegistryTag(UE::Workspace::ExportsWorkspaceItemsRegistryTag, TagValue, UObject::FAssetRegistryTag::TT_Hidden));
-		}
 	}
 }
