Fix assets not recompiling on shared variable changes.

#tests PIE

diff --git a/Engine/Plugins/Experimental/UAF/UAF/Source/UAFUncookedOnly/Private/Entries/AnimNextSharedVariablesEntry.cpp b/Engine/Plugins/Experimental/UAF/UAF/Source/UAFUncookedOnly/Private/Entries/AnimNextSharedVariablesEntry.cpp
--- a/Engine/Plugins/Experimental/UAF/UAF/Source/UAFUncookedOnly/Private/Entries/AnimNextSharedVariablesEntry.cpp
+++ b/Engine/Plugins/Experimental/UAF/UAF/Source/UAFUncookedOnly/Private/Entries/AnimNextSharedVariablesEntry.cpp
@@ -132,6 +132,7 @@ void UAnimNextSharedVariablesEntry::HandleAssetModified(UAnimNextRigVMAssetEdito
 	case EAnimNextEditorDataNotifType::EntryAccessSpecifierChanged:
 	case EAnimNextEditorDataNotifType::VariableTypeChanged:
 	case EAnimNextEditorDataNotifType::VariableDefaultValueChanged:
+	case EAnimNextEditorDataNotifType::VariableBindingChanged:
 		if(UAnimNextRigVMAssetEditorData* EditorData = GetTypedOuter<UAnimNextRigVMAssetEditorData>())
 		{
 			EditorData->RequestAutoVMRecompilation();
diff --git a/Engine/Plugins/Experimental/UAF/UAFStateTree/Source/UAFStateTreeUncookedOnly/Private/AnimNextStateTreeEditorData.cpp b/Engine/Plugins/Experimental/UAF/UAFStateTree/Source/UAFStateTreeUncookedOnly/Private/AnimNextStateTreeEditorData.cpp
--- a/Engine/Plugins/Experimental/UAF/UAFStateTree/Source/UAFStateTreeUncookedOnly/Private/AnimNextStateTreeEditorData.cpp
+++ b/Engine/Plugins/Experimental/UAF/UAFStateTree/Source/UAFStateTreeUncookedOnly/Private/AnimNextStateTreeEditorData.cpp
@@ -57,6 +57,7 @@ void UAnimNextStateTreeTreeEditorData::HandleStateTreeAssetChanges(UAnimNextRigV
 	case EAnimNextEditorDataNotifType::EntryAccessSpecifierChanged:
 	case EAnimNextEditorDataNotifType::VariableTypeChanged:
 	case EAnimNextEditorDataNotifType::VariableDefaultValueChanged:
+	case EAnimNextEditorDataNotifType::VariableBindingChanged:
 		{			
 			const UStateTree* StateTree = GetTypedOuter<UStateTree>();
 			UE::StateTree::Delegates::OnParametersChanged.Broadcast(*StateTree);	
