Downgrade check to ensure on statetree failure

#tests PIE

diff --git a/Engine/Plugins/Experimental/UAF/UAFStateTree/Source/UAFStateTree/Private/AnimStateTreeTrait.cpp b/Engine/Plugins/Experimental/UAF/UAFStateTree/Source/UAFStateTree/Private/AnimStateTreeTrait.cpp
--- a/Engine/Plugins/Experimental/UAF/UAFStateTree/Source/UAFStateTree/Private/AnimStateTreeTrait.cpp
+++ b/Engine/Plugins/Experimental/UAF/UAFStateTree/Source/UAFStateTree/Private/AnimStateTreeTrait.cpp
@@ -105,9 +105,16 @@ namespace UE::UAF
 
 			FAnimNextGraphInstance& OwnerGraphInstance = Binding.GetTraitPtr().GetNodeInstance()->GetOwner();
 			StateTreeExecutionContext.SetOuterTraceId(OwnerGraphInstance.GetUniqueId());
-#if DO_CHECK
-			check(OwnerGraphInstance.LayoutMatches(InstanceData->StateTree->GetDefaultParameters()));
-#endif
+
+#if ENABLE_ANIM_DEBUG 
+			if (!OwnerGraphInstance.LayoutMatches(InstanceData->StateTree->GetDefaultParameters()))
+			{
+				UE_LOG(LogAnimation, Error, TEXT("Anim StateTree Parameter Layout Mismatch: %s - StateTree: %s"), *Owner->GetFName().ToString(), *InstanceData->StateTree->GetOuter()->GetFName().ToString());
+				StateTreeExecutionContext.Stop(EStateTreeRunStatus::Failed);
+				return;
+			}
+#endif // ENABLE_ANIM_DEBUG 
+
 			const FInstancedPropertyBag& StateTreeParameters = InstanceData->StateTree->GetDefaultParameters();
 			InstanceData->StateTreeExternalParameters.Reset();
 			
