UAF StateTree autocomplete on graph timeline complete

#tests PIE

diff --git a/Engine/Plugins/Experimental/UAF/UAFStateTree/Source/UAFStateTree/Internal/Tasks/AnimNextStateTreeGraphInstanceTask.h b/Engine/Plugins/Experimental/UAF/UAFStateTree/Source/UAFStateTree/Internal/Tasks/AnimNextStateTreeGraphInstanceTask.h
--- a/Engine/Plugins/Experimental/UAF/UAFStateTree/Source/UAFStateTree/Internal/Tasks/AnimNextStateTreeGraphInstanceTask.h
+++ b/Engine/Plugins/Experimental/UAF/UAFStateTree/Source/UAFStateTree/Internal/Tasks/AnimNextStateTreeGraphInstanceTask.h
@@ -4,6 +4,7 @@
 
 #include "AnimNextStateTreeTypes.h"
 #include "Graph/AnimNextAnimationGraph.h"
+#include <limits>
 #include "AlphaBlend.h"
 #include "Factory/AnimNextFactoryParams.h"
 #include "StructUtils/PropertyBag.h"
@@ -45,10 +46,26 @@ struct UAFSTATETREE_API FAnimNextGraphInstanceTaskInstanceData
 	UPROPERTY(EditAnywhere, Category = Animation)
 	bool bContinueTicking = true;
 
-	// Current playback ratio (debug)
+	// How early to trigger complete. Setting this allows for a blend out while the timeline is still playing.
+	UPROPERTY(EditAnywhere, Category = Animation)
+	float CompleteBlendOutTime = 0.0f;
+
+	// Current playback ratio 
 	UPROPERTY(VisibleAnywhere, Category = Animation)
 	float PlaybackRatio = 1.0f;
 
+	// Current time remaining
+	UPROPERTY(VisibleAnywhere, Category = Animation)
+	float TimeLeft = std::numeric_limits<float>::infinity();
+
+	// Current graph duration
+	UPROPERTY(VisibleAnywhere, Category = Animation)
+	float Duration = 0.0f;
+
+	// Current looping status
+	UPROPERTY(VisibleAnywhere, Category = Animation)
+	bool bIsLooping = false;
+
 #if WITH_EDITORONLY_DATA
 	bool Serialize(FArchive& Ar);
 	void PostSerialize(const FArchive& Ar);
diff --git a/Engine/Plugins/Experimental/UAF/UAFStateTree/Source/UAFStateTree/Private/AnimNextStateTreeContext.cpp b/Engine/Plugins/Experimental/UAF/UAFStateTree/Source/UAFStateTree/Private/AnimNextStateTreeContext.cpp
--- a/Engine/Plugins/Experimental/UAF/UAFStateTree/Source/UAFStateTree/Private/AnimNextStateTreeContext.cpp
+++ b/Engine/Plugins/Experimental/UAF/UAFStateTree/Source/UAFStateTree/Private/AnimNextStateTreeContext.cpp
@@ -1,6 +1,8 @@
 // Copyright Epic Games, Inc. All Rights Reserved.
 
 #include "AnimNextStateTreeContext.h"
+
+#include <limits>
 #include "Factory/AnimGraphFactory.h"
 #include "TraitCore/NodeInstance.h"
 #include "TraitCore/TraitStackBinding.h"
@@ -35,7 +37,7 @@ bool FAnimNextStateTreeTraitContext::PushAssetOntoBlendStack(TNonNullPtr<UObject
 	return true;
 }
 
-float FAnimNextStateTreeTraitContext::QueryPlaybackRatio(TNonNullPtr<UObject> InAsset) const
+float FAnimNextStateTreeTraitContext::QueryPlaybackRatio() const
 {
 	UE::UAF::TTraitBinding<UE::UAF::ITimeline> Timeline;
 	if (Binding->GetStackInterface<UE::UAF::ITimeline>(Timeline))
@@ -46,6 +48,39 @@ float FAnimNextStateTreeTraitContext::QueryPlaybackRatio(TNonNullPtr<UObject> In
 	return 1.0f;
 }
 
+float FAnimNextStateTreeTraitContext::QueryTimeLeft() const
+{
+	UE::UAF::TTraitBinding<UE::UAF::ITimeline> Timeline;
+	if (Binding->GetStackInterface<UE::UAF::ITimeline>(Timeline))
+	{
+		return Timeline.GetState(*Context).GetTimeLeft();
+	}
+
+	return std::numeric_limits<float>::infinity();
+}
+
+float FAnimNextStateTreeTraitContext::QueryDuration() const
+{
+	UE::UAF::TTraitBinding<UE::UAF::ITimeline> Timeline;
+	if (Binding->GetStackInterface<UE::UAF::ITimeline>(Timeline))
+	{
+		return Timeline.GetState(*Context).GetDuration();
+	}
+
+	return 0.0f;
+}
+
+bool FAnimNextStateTreeTraitContext::QueryIsLooping() const
+{
+	UE::UAF::TTraitBinding<UE::UAF::ITimeline> Timeline;
+	if (Binding->GetStackInterface<UE::UAF::ITimeline>(Timeline))
+	{
+		return Timeline.GetState(*Context).IsLooping();
+	}
+
+	return false;
+}
+
 FAnimNextGraphInstance& FAnimNextStateTreeTraitContext::GetGraphInstance() const
 {
 	return Binding->GetTraitPtr().GetNodeInstance()->GetOwner();
diff --git a/Engine/Plugins/Experimental/UAF/UAFStateTree/Source/UAFStateTree/Private/AnimNextStateTreeContext.h b/Engine/Plugins/Experimental/UAF/UAFStateTree/Source/UAFStateTree/Private/AnimNextStateTreeContext.h
--- a/Engine/Plugins/Experimental/UAF/UAFStateTree/Source/UAFStateTree/Private/AnimNextStateTreeContext.h
+++ b/Engine/Plugins/Experimental/UAF/UAFStateTree/Source/UAFStateTree/Private/AnimNextStateTreeContext.h
@@ -34,7 +34,10 @@ struct FAnimNextStateTreeTraitContext
 	{}
 
 	bool PushAssetOntoBlendStack(TNonNullPtr<UObject> InAsset, const FAlphaBlendArgs& InBlendArguments, const FAnimNextFactoryParams& InFactoryParams) const;
-	float QueryPlaybackRatio(TNonNullPtr<UObject> InAsset) const;
+	float QueryPlaybackRatio() const;
+	float QueryTimeLeft() const;
+	float QueryDuration() const;
+	bool QueryIsLooping() const;
 	FAnimNextGraphInstance& GetGraphInstance() const;
 
 	UE::UAF::FExecutionContext& GetContext() const
diff --git a/Engine/Plugins/Experimental/UAF/UAFStateTree/Source/UAFStateTree/Private/Tasks/AnimNextStateTreeGraphInstanceTask.cpp b/Engine/Plugins/Experimental/UAF/UAFStateTree/Source/UAFStateTree/Private/Tasks/AnimNextStateTreeGraphInstanceTask.cpp
--- a/Engine/Plugins/Experimental/UAF/UAFStateTree/Source/UAFStateTree/Private/Tasks/AnimNextStateTreeGraphInstanceTask.cpp
+++ b/Engine/Plugins/Experimental/UAF/UAFStateTree/Source/UAFStateTree/Private/Tasks/AnimNextStateTreeGraphInstanceTask.cpp
@@ -111,7 +111,21 @@ EStateTreeRunStatus FAnimNextStateTreeGraphInstanceTask::Tick(FStateTreeExecutio
 	{
 		FAnimNextStateTreeTraitContext& ExecContext = Context.GetExternalData(TraitContextHandle);
 
-		InstanceData.PlaybackRatio = ExecContext.QueryPlaybackRatio(InstanceData.Asset.Get());
+		InstanceData.PlaybackRatio = ExecContext.QueryPlaybackRatio();
+		InstanceData.Duration = ExecContext.QueryDuration();
+		InstanceData.TimeLeft = ExecContext.QueryTimeLeft();
+		InstanceData.bIsLooping = ExecContext.QueryIsLooping();
+
+		// We can get a fallback timeline with duration of 0.0f during graph init / compile.
+		// Avoid finishing the task for this edge case
+		// @TODO: Later on consider making 0.0f duration invalid in timeline state.
+		if (InstanceData.Duration != 0.0f)
+		{
+			if (InstanceData.TimeLeft - InstanceData.CompleteBlendOutTime <= 0.0f && !InstanceData.bIsLooping)
+			{
+				Context.FinishTask(*this, EStateTreeFinishTaskType::Succeeded);
+			}
+		}
 
 		return EStateTreeRunStatus::Running;
 	}
