Add missing IUpdate propagation to StateTree

#tests PIE

diff --git a/Engine/Plugins/Experimental/UAF/UAFStateTree/Source/UAFStateTree/Private/AnimStateTreeTrait.cpp b/Engine/Plugins/Experimental/UAF/UAFStateTree/Source/UAFStateTree/Private/AnimStateTreeTrait.cpp
--- a/Engine/Plugins/Experimental/UAF/UAFStateTree/Source/UAFStateTree/Private/AnimStateTreeTrait.cpp
+++ b/Engine/Plugins/Experimental/UAF/UAFStateTree/Source/UAFStateTree/Private/AnimStateTreeTrait.cpp
@@ -247,6 +247,9 @@ namespace UE::UAF
 				StateTreeExecutionContext.Tick(TraitState.GetDeltaTime());
 			}
 		}
+
+		// Update the traits below us, we updated first so we can push transitions
+		IUpdate::PreUpdate(Context, Binding, TraitState);
 	}
 
 	void FStateTreeTrait::AddReferencedObjects(const FExecutionContext& Context, const TTraitBinding<IGarbageCollection>& Binding, FReferenceCollector& Collector) const
