Fix UAF leaking callbacks causing UAF state tree selction to be cleared.
- This prevented things like delete from being performable on right click.

[StateTree] Fix ST leaking `GetOnStateNodesChanged` delegate callbacks.

#tests editor

diff --git a/Engine/Plugins/Experimental/UAF/UAFStateTree/Source/UAFStateTreeEditor/Private/AnimNextStateTreeEditorModule.cpp b/Engine/Plugins/Experimental/UAF/UAFStateTree/Source/UAFStateTreeEditor/Private/AnimNextStateTreeEditorModule.cpp
--- a/Engine/Plugins/Experimental/UAF/UAFStateTree/Source/UAFStateTreeEditor/Private/AnimNextStateTreeEditorModule.cpp
+++ b/Engine/Plugins/Experimental/UAF/UAFStateTree/Source/UAFStateTreeEditor/Private/AnimNextStateTreeEditorModule.cpp
@@ -95,7 +95,10 @@ void FAnimNextStateTreeEditorModule::StartupModule()
 
 				TWeakPtr<SWidget> WeakStateTreeViewWidget = StateTreeViewWidget;
 				TWeakPtr<FStateTreeViewModel> WeakViewModel = StateTreeViewModel;
-				StateTreeViewModel->GetOnSelectionChanged().AddSPLambda(&StateTreeViewWidget.Get(), [WeakStateTreeViewWidget, WeakWorkspaceEditor, WeakViewModel](const TArray<TWeakObjectPtr<UStateTreeState>>& SelectedStates)
+
+				// @TODO: This works but it would be more idiomatic to have the remove occur at a point of de-initialization. This allows for the context object to be destroyed between user sessions without leaking delegate callbacks.
+				StateTreeViewModel->GetOnSelectionChanged().RemoveAll(&InContext.WorkspaceEditor.Get());
+				StateTreeViewModel->GetOnSelectionChanged().AddSPLambda(&InContext.WorkspaceEditor.Get(), [WeakStateTreeViewWidget, WeakWorkspaceEditor, WeakViewModel](const TArray<TWeakObjectPtr<UStateTreeState>>& SelectedStates)
 				{
 					if (TSharedPtr<Workspace::IWorkspaceEditor> SharedWorkspaceEditor = WeakWorkspaceEditor.Pin())
 					{
@@ -108,12 +111,15 @@ void FAnimNextStateTreeEditorModule::StartupModule()
 							}
 						}
 
-						SharedWorkspaceEditor->SetGlobalSelection(WeakStateTreeViewWidget, UE::Workspace::FOnClearGlobalSelection::CreateLambda([WeakViewModel]()
+						SharedWorkspaceEditor->SetGlobalSelection(WeakStateTreeViewWidget, UE::Workspace::FOnClearGlobalSelection::CreateLambda([WeakViewModel, WeakStateTreeViewWidget]()
 						{ 
-							if (const TSharedPtr<FStateTreeViewModel> SharedViewModel = WeakViewModel.Pin()) 
-							{ 
-								SharedViewModel->ClearSelection();
-							} 
+							if (WeakStateTreeViewWidget.IsValid())
+							{
+								if (const TSharedPtr<FStateTreeViewModel> SharedViewModel = WeakViewModel.Pin())
+								{
+									SharedViewModel->ClearSelection();
+								}
+							}
 						}));						
 						SharedWorkspaceEditor->SetDetailsObjects(Selected);
 					}					
diff --git a/Engine/Plugins/Runtime/StateTree/Source/StateTreeEditorModule/Private/StateTreeEditorMode.cpp b/Engine/Plugins/Runtime/StateTree/Source/StateTreeEditorModule/Private/StateTreeEditorMode.cpp
--- a/Engine/Plugins/Runtime/StateTree/Source/StateTreeEditorModule/Private/StateTreeEditorMode.cpp
+++ b/Engine/Plugins/Runtime/StateTree/Source/StateTreeEditorModule/Private/StateTreeEditorMode.cpp
@@ -255,6 +255,7 @@ void UStateTreeEditorMode::OnStateTreeChanged()
 					OldViewModel->GetOnStateAdded().RemoveAll(this);
 					OldViewModel->GetOnStatesRemoved().RemoveAll(this);
 					OldViewModel->GetOnStatesMoved().RemoveAll(this);
+					OldViewModel->GetOnStateNodesChanged().RemoveAll(this);
 					OldViewModel->GetOnSelectionChanged().RemoveAll(this);
 					OldViewModel->GetOnBringNodeToFocus().RemoveAll(this);
 				}
