Fix UAF state tree RigVM conditions only working on enter conditions

#tests editor
#Synced-CL 48531498

diff --git a/Engine/Plugins/Experimental/UAF/UAFStateTree/Source/UAFStateTreeUncookedOnly/Private/AnimNextStateTree_EditorData.cpp b/Engine/Plugins/Experimental/UAF/UAFStateTree/Source/UAFStateTreeUncookedOnly/Private/AnimNextStateTree_EditorData.cpp
--- a/Engine/Plugins/Experimental/UAF/UAFStateTree/Source/UAFStateTreeUncookedOnly/Private/AnimNextStateTree_EditorData.cpp
+++ b/Engine/Plugins/Experimental/UAF/UAFStateTree/Source/UAFStateTreeUncookedOnly/Private/AnimNextStateTree_EditorData.cpp
@@ -15,6 +15,7 @@
 #include "StateTreeCompilerLog.h"
 #include "StateTreeEditingSubsystem.h"
 #include "StateTreeEditorData.h"
+#include "StateTreeState.h"
 #include "Entries/AnimNextVariableEntry.h"
 #include "Entries/AnimNextAnimationGraphEntry.h"
 #include "Graph/AnimNextAnimationGraphSchema.h"
@@ -105,6 +106,25 @@ void UAnimNextStateTree_EditorData::BuildFunctionHeadersContext(const FRigVMComp
 				}
 			}
 
+			for (const FStateTreeTransition& Transition : State.Transitions)
+			{
+				for (const FStateTreeEditorNode& Node : Transition.Conditions)
+				{
+					// Manually check for each child node type, UStruct's do not reflect interfaces so we can't use an interface here.
+					if (const FAnimNextStateTreeRigVMConditionBase* Condition = Node.Node.GetPtr<FAnimNextStateTreeRigVMConditionBase>())
+					{
+						FStateTreeBindableStructDesc Desc;
+						Desc.Struct = Condition->GetInstanceDataType();
+						Desc.Name = Node.GetName();
+						Desc.ID = Node.ID;
+						Desc.DataSource = EStateTreeBindableStructSource::Transition;
+
+						// Graph gen isn't const since we need to cache generated names in the node
+						const_cast<FAnimNextStateTreeRigVMConditionBase*>(Condition)->GetProgrammaticFunctionHeaders(ProgrammaticFunctionHeaderParams, &State, Desc);
+					}
+				}
+			}
+
 			for (const FStateTreeEditorNode& Node : State.Tasks)
 			{
 				// Manually check for each child node type, UStruct's do not reflect interfaces so we can't use an interface here.
