Fix UAF ST functions counting input variables as params

#tests PIE
#Synced-CL 50015298

diff --git a/Engine/Plugins/Experimental/UAF/UAFStateTree/Source/UAFStateTree/Private/Conditions/AnimNextStateTreeRigVMConditionBase.cpp b/Engine/Plugins/Experimental/UAF/UAFStateTree/Source/UAFStateTree/Private/Conditions/AnimNextStateTreeRigVMConditionBase.cpp
--- a/Engine/Plugins/Experimental/UAF/UAFStateTree/Source/UAFStateTree/Private/Conditions/AnimNextStateTreeRigVMConditionBase.cpp
+++ b/Engine/Plugins/Experimental/UAF/UAFStateTree/Source/UAFStateTree/Private/Conditions/AnimNextStateTreeRigVMConditionBase.cpp
@@ -105,7 +105,11 @@ void FAnimNextStateTreeRigVMConditionBase::PostEditNodeChangeChainProperty(const
 						PropertyDesc.PropertyFlags &= ~CPF_Edit;	// Hide the result from the UI
 						ResultIndex = PropertyDescs.Num();
 					}
-					PropertyDescs.Add(PropertyDesc);
+
+					if (!Argument.bIsInputVariable)
+					{
+						PropertyDescs.Add(PropertyDesc);
+					}
 				}
 			}
 
diff --git a/Engine/Plugins/Experimental/UAF/UAFStateTree/Source/UAFStateTree/Private/Tasks/AnimNextStateTreeRigVMTaskBase.cpp b/Engine/Plugins/Experimental/UAF/UAFStateTree/Source/UAFStateTree/Private/Tasks/AnimNextStateTreeRigVMTaskBase.cpp
--- a/Engine/Plugins/Experimental/UAF/UAFStateTree/Source/UAFStateTree/Private/Tasks/AnimNextStateTreeRigVMTaskBase.cpp
+++ b/Engine/Plugins/Experimental/UAF/UAFStateTree/Source/UAFStateTree/Private/Tasks/AnimNextStateTreeRigVMTaskBase.cpp
@@ -110,7 +110,11 @@ void FAnimNextStateTreeRigVMTaskBase::PostEditNodeChangeChainProperty(const FPro
 						PropertyDesc.PropertyFlags &= ~CPF_Edit;	// Hide the results from the UI
 						ResultIndex = PropertyDescs.Num();
 					}
-					PropertyDescs.Add(PropertyDesc);
+
+					if (!Argument.bIsInputVariable)
+					{
+						PropertyDescs.Add(PropertyDesc);
+					}
 				}
 			}
 		
