Add ability to call UAF ST function on state tick.

#tests editor, PIE

#Synced-CL 50972851

diff --git a/Engine/Plugins/Experimental/UAF/UAFStateTree/Source/UAFStateTree/Internal/Tasks/AnimNextStateTreeRigVMTaskBase.h b/Engine/Plugins/Experimental/UAF/UAFStateTree/Source/UAFStateTree/Internal/Tasks/AnimNextStateTreeRigVMTaskBase.h
--- a/Engine/Plugins/Experimental/UAF/UAFStateTree/Source/UAFStateTree/Internal/Tasks/AnimNextStateTreeRigVMTaskBase.h
+++ b/Engine/Plugins/Experimental/UAF/UAFStateTree/Source/UAFStateTree/Internal/Tasks/AnimNextStateTreeRigVMTaskBase.h
@@ -23,6 +23,13 @@ struct FAnimNextStateTreeTraitContext;
 
 #define UE_API UAFSTATETREE_API
 
+UENUM()
+enum class EStateTreeFunctionCallFrequency
+{
+	OnEnter,
+	OnTick
+};
+
 USTRUCT()
 struct UE_API FAnimNextStateTreeRigVMTaskInstanceData
 {
@@ -51,6 +58,7 @@ public:
 	//~ Begin FStateTreeTaskBase Interface
 	virtual bool Link(FStateTreeLinker& Linker) override;
 	virtual EStateTreeRunStatus EnterState(FStateTreeExecutionContext& Context, const FStateTreeTransitionResult& Transition) const override;
+	virtual EStateTreeRunStatus Tick(FStateTreeExecutionContext& Context, const float DeltaTime) const override;
 	//~ End FStateTreeTaskBase Interface
 
 	//~ Begin FStateTreeNodeBase Interface
@@ -83,6 +91,9 @@ public:
 	UPROPERTY(EditAnywhere, Category = "Call Function", meta = (Hidden))
 	FUAFCallFunctionInfo CallFunctionInfo;
 
+	UPROPERTY(EditAnywhere, Category = "Call Function")
+	EStateTreeFunctionCallFrequency CallFunctionFrequency = EStateTreeFunctionCallFrequency::OnEnter;
+
 	/** Owning state name, populated during programmatic graph creation */
 	UPROPERTY()
 	FName StateName = NAME_None;
@@ -94,10 +105,18 @@ public:
 	UPROPERTY()
 	int32 ResultIndex = INDEX_NONE;
 
+public:
+
 	UE_API void PostSerialize(const FArchive& Ar);
 
 public:
+
 	TStateTreeExternalDataHandle<FAnimNextStateTreeTraitContext> TraitContextHandle;
+
+protected:
+
+	/** Internal impl to call some RigVM function */
+	EStateTreeRunStatus CallFunction(FStateTreeExecutionContext& Context) const;
 };
 
 #if WITH_EDITORONLY_DATA
diff --git a/Engine/Plugins/Experimental/UAF/UAFStateTree/Source/UAFStateTree/Private/Conditions/AnimNextStateTreeRigVMConditionBase.cpp b/Engine/Plugins/Experimental/UAF/UAFStateTree/Source/UAFStateTree/Private/Conditions/AnimNextStateTreeRigVMConditionBase.cpp
--- a/Engine/Plugins/Experimental/UAF/UAFStateTree/Source/UAFStateTree/Private/Conditions/AnimNextStateTreeRigVMConditionBase.cpp
+++ b/Engine/Plugins/Experimental/UAF/UAFStateTree/Source/UAFStateTree/Private/Conditions/AnimNextStateTreeRigVMConditionBase.cpp
@@ -118,8 +118,6 @@ void FAnimNextStateTreeRigVMConditionBase::PostEditNodeChangeChainProperty(const
 			}
 
 			// @TODO: Constrain functions we can pick to those returning only a bool
-			ensure(ResultIndex != INDEX_NONE);
-
 			InstanceData->Parameters.AddProperties(PropertyDescs);
 		}
 	}
diff --git a/Engine/Plugins/Experimental/UAF/UAFStateTree/Source/UAFStateTree/Private/Tasks/AnimNextStateTreeRigVMTaskBase.cpp b/Engine/Plugins/Experimental/UAF/UAFStateTree/Source/UAFStateTree/Private/Tasks/AnimNextStateTreeRigVMTaskBase.cpp
--- a/Engine/Plugins/Experimental/UAF/UAFStateTree/Source/UAFStateTree/Private/Tasks/AnimNextStateTreeRigVMTaskBase.cpp
+++ b/Engine/Plugins/Experimental/UAF/UAFStateTree/Source/UAFStateTree/Private/Tasks/AnimNextStateTreeRigVMTaskBase.cpp
@@ -35,6 +35,26 @@ bool FAnimNextStateTreeRigVMTaskBase::Link(FStateTreeLinker& Linker)
 }
 
 EStateTreeRunStatus FAnimNextStateTreeRigVMTaskBase::EnterState(FStateTreeExecutionContext& Context, const FStateTreeTransitionResult& Transition) const
+{
+	if (CallFunctionFrequency != EStateTreeFunctionCallFrequency::OnEnter)
+	{
+		return Super::EnterState(Context, Transition);
+	}
+
+	return CallFunction(Context);
+}
+
+EStateTreeRunStatus FAnimNextStateTreeRigVMTaskBase::Tick(FStateTreeExecutionContext& Context, const float DeltaTime) const
+{
+	if (CallFunctionFrequency != EStateTreeFunctionCallFrequency::OnTick)
+	{
+		return Super::Tick(Context, DeltaTime);
+	}
+
+	return CallFunction(Context);
+}
+
+EStateTreeRunStatus FAnimNextStateTreeRigVMTaskBase::CallFunction(FStateTreeExecutionContext& Context) const
 {
 	if (FAnimNextStateTreeRigVMTaskInstanceData* InstanceData = Context.GetInstanceDataPtr<FAnimNextStateTreeRigVMTaskInstanceData>(*this))
 	{
@@ -66,7 +86,6 @@ EStateTreeRunStatus FAnimNextStateTreeRigVMTaskBase::EnterState(FStateTreeExecut
 			FAnimNextGraphContextData ContextData(GraphInstance, CallFunctionInfo.FunctionEvent, AnimNextContext.GetDeltaTime<float>());
 			UE::UAF::FScopedExecuteContextData ContextDataScope(AnimNextContext, ContextData);
 
-			// @TODO: Add option to execute per tick?
 			RigVMAsset->CallFunctionHandle(InstanceData->FunctionHandle, ExtendedExecuteContext, InstanceData->Parameters);
 
 			// @TODO: For now we use a fire & forget model. But in the future we might want users to be able to bind to the result value from a function
