Fix UAF ST RigVM function params do not update when params are changed after binding.

#tests PIE, editor

#Synced-CL 50972851

diff --git a/Engine/Plugins/Experimental/UAF/UAFStateTree/Source/UAFStateTree/Internal/Conditions/AnimNextStateTreeRigVMConditionBase.h b/Engine/Plugins/Experimental/UAF/UAFStateTree/Source/UAFStateTree/Internal/Conditions/AnimNextStateTreeRigVMConditionBase.h
--- a/Engine/Plugins/Experimental/UAF/UAFStateTree/Source/UAFStateTree/Internal/Conditions/AnimNextStateTreeRigVMConditionBase.h
+++ b/Engine/Plugins/Experimental/UAF/UAFStateTree/Source/UAFStateTree/Internal/Conditions/AnimNextStateTreeRigVMConditionBase.h
@@ -64,7 +64,10 @@ public:
 	
 #if WITH_EDITOR
 	/** Adds condition function headers to parent StateTree RigVM */
-	virtual void GetProgrammaticFunctionHeaders(FAnimNextStateTreeProgrammaticFunctionHeaderParams& InProgrammaticFunctionHeaderParams, const UStateTreeState* State, const FStateTreeBindableStructDesc& Desc);
+	virtual void GetProgrammaticFunctionHeaders(FAnimNextStateTreeProgrammaticFunctionHeaderParams& InProgrammaticFunctionHeaderParams, UStateTreeState* State, FStateTreeDataView InstanceDataView, const FStateTreeBindableStructDesc& Desc);
+
+	/** Regenerates the instance data params */
+	void RegenerateInstanceDataFunctionParams(TNotNull<FInstanceDataType*> InstanceData);
 #endif
 
 public:
diff --git a/Engine/Plugins/Experimental/UAF/UAFStateTree/Source/UAFStateTree/Internal/Tasks/AnimNextStateTreeRigVMTaskBase.h b/Engine/Plugins/Experimental/UAF/UAFStateTree/Source/UAFStateTree/Internal/Tasks/AnimNextStateTreeRigVMTaskBase.h
--- a/Engine/Plugins/Experimental/UAF/UAFStateTree/Source/UAFStateTree/Internal/Tasks/AnimNextStateTreeRigVMTaskBase.h
+++ b/Engine/Plugins/Experimental/UAF/UAFStateTree/Source/UAFStateTree/Internal/Tasks/AnimNextStateTreeRigVMTaskBase.h
@@ -71,7 +71,10 @@ public:
 	
 #if WITH_EDITOR
 	/** Adds Task function headers to parent StateTree RigVM */
-	virtual void GetProgrammaticFunctionHeaders(FAnimNextStateTreeProgrammaticFunctionHeaderParams& InProgrammaticFunctionHeaderParams, const UStateTreeState* State, const FStateTreeBindableStructDesc& Desc);
+	virtual void GetProgrammaticFunctionHeaders(FAnimNextStateTreeProgrammaticFunctionHeaderParams& InProgrammaticFunctionHeaderParams, UStateTreeState* State, FStateTreeDataView InstanceDataView, const FStateTreeBindableStructDesc& Desc);
+
+	/** Regenerates the instance data params */
+	void RegenerateInstanceDataFunctionParams(TNotNull<FInstanceDataType*> InstanceData);
 #endif
 
 public:
diff --git a/Engine/Plugins/Experimental/UAF/UAFStateTree/Source/UAFStateTree/Private/Conditions/AnimNextStateTreeRigVMConditionBase.cpp b/Engine/Plugins/Experimental/UAF/UAFStateTree/Source/UAFStateTree/Private/Conditions/AnimNextStateTreeRigVMConditionBase.cpp
--- a/Engine/Plugins/Experimental/UAF/UAFStateTree/Source/UAFStateTree/Private/Conditions/AnimNextStateTreeRigVMConditionBase.cpp
+++ b/Engine/Plugins/Experimental/UAF/UAFStateTree/Source/UAFStateTree/Private/Conditions/AnimNextStateTreeRigVMConditionBase.cpp
@@ -92,49 +92,102 @@ void FAnimNextStateTreeRigVMConditionBase::PostEditNodeChangeChainProperty(const
 		// Function Name selected has changed. Update param struct & result / event names used during execution.
 		if (FAnimNextStateTreeRigVMConditionInstanceData* InstanceData = InstanceDataView.GetMutablePtr<FAnimNextStateTreeRigVMConditionInstanceData>())
 		{
-			InstanceData->Parameters.Reset();
-			ResultIndex = INDEX_NONE;
+			RegenerateInstanceDataFunctionParams(InstanceData);
+		}
+	}
+}
+
+void FAnimNextStateTreeRigVMConditionBase::GetProgrammaticFunctionHeaders(FAnimNextStateTreeProgrammaticFunctionHeaderParams& InProgrammaticFunctionHeaderParams, UStateTreeState* State, FStateTreeDataView InstanceDataView, const FStateTreeBindableStructDesc& Desc)
+{
+	using namespace UE::UAF::UncookedOnly;
+
+	FAnimNextGetFunctionHeaderCompileContext& OutCompileContext = InProgrammaticFunctionHeaderParams.OutCompileContext;
+
+	StateName = State->Name;
+	NodeId = Desc.ID;
+
+	// Validate that our cached instance parameters match the latest RigVM signature. Update if needed.
+	UStateTree* OwningStateTree = State->GetTypedOuter<UStateTree>();
+	if (FAnimNextStateTreeRigVMConditionInstanceData* InstanceData = InstanceDataView.GetMutablePtr<FAnimNextStateTreeRigVMConditionInstanceData>())
+	{
+		auto IsBindingOutdated = [&]()
+		{
+			if (!InstanceData->Parameters.IsValid())
+			{
+				return false;
+			}
+
+			TConstArrayView<FPropertyBagPropertyDesc> CachedDescs = InstanceData->Parameters.GetPropertyBagStruct()->GetPropertyDescs();
+
+			// Our function header might be out of date. Get latest
+			CallFunctionInfo.FunctionHeader = FRigVMGraphFunctionHeader::FindGraphFunctionHeader(CallFunctionInfo.FunctionHeader.LibraryPointer);
+			int32 RigVMFunctionBindableArgumentIndex = 0;
 
-			TArray<FPropertyBagPropertyDesc> PropertyDescs;
-			PropertyDescs.Reserve(CallFunctionInfo.FunctionHeader.Arguments.Num());
 			for (const FRigVMGraphFunctionArgument& Argument : CallFunctionInfo.FunctionHeader.Arguments)
 			{
-				if (Argument.Direction == ERigVMPinDirection::Input || Argument.Direction == ERigVMPinDirection::Output)
+				const bool bIsBindable = Argument.Direction == ERigVMPinDirection::Input || Argument.Direction == ERigVMPinDirection::Output;
+				if (bIsBindable && !Argument.bIsInputVariable)
 				{
 					FString CppTypeString = Argument.CPPType.ToString();
 					UObject* CppTypeObject = Argument.CPPTypeObject.Get();
 					FPropertyBagPropertyDesc PropertyDesc = FRigVMMemoryStorageStruct::GeneratePropertyBagDescriptor(FRigVMPropertyDescription(Argument.Name, CppTypeString, CppTypeObject, Argument.DefaultValue));
-					if (Argument.Direction == ERigVMPinDirection::Output && Argument.CPPType == RigVMTypeUtils::BoolTypeName)
-					{
-						PropertyDesc.PropertyFlags &= ~CPF_Edit;	// Hide the result from the UI
-						ResultIndex = PropertyDescs.Num();
-					}
 
-					if (!Argument.bIsInputVariable)
+					if (!CachedDescs.IsValidIndex(RigVMFunctionBindableArgumentIndex) || CachedDescs[RigVMFunctionBindableArgumentIndex] != PropertyDesc)
 					{
-						PropertyDescs.Add(PropertyDesc);
+						return true;
 					}
+
+					RigVMFunctionBindableArgumentIndex++;
 				}
 			}
 
-			// @TODO: Constrain functions we can pick to those returning only a bool
-			InstanceData->Parameters.AddProperties(PropertyDescs);
+			const int32 NumBindableArguments = RigVMFunctionBindableArgumentIndex;
+			if (NumBindableArguments != CachedDescs.Num())
+			{
+				return true;
+			}
+
+			return false;
+		};
+
+		if (IsBindingOutdated())
+		{
+			RegenerateInstanceDataFunctionParams(InstanceData);
 		}
 	}
+
+	FAnimNextProgrammaticFunctionHeader AnimNextFunctionHeader = {};
+	AnimNextFunctionHeader.FunctionHeader = CallFunctionInfo.FunctionHeader;
+	OutCompileContext.AddUniqueFunctionHeader(AnimNextFunctionHeader);
 }
 
-void FAnimNextStateTreeRigVMConditionBase::GetProgrammaticFunctionHeaders(FAnimNextStateTreeProgrammaticFunctionHeaderParams& InProgrammaticFunctionHeaderParams, const UStateTreeState* State, const FStateTreeBindableStructDesc& Desc)
+void FAnimNextStateTreeRigVMConditionBase::RegenerateInstanceDataFunctionParams(TNotNull<FInstanceDataType*> InstanceData)
 {
-	using namespace UE::UAF::UncookedOnly;
+	InstanceData->Parameters.Reset();
+	ResultIndex = INDEX_NONE;
 
-	FAnimNextGetFunctionHeaderCompileContext& OutCompileContext = InProgrammaticFunctionHeaderParams.OutCompileContext;
-	
-	StateName = State->Name;
-	NodeId = Desc.ID;
+	TArray<FPropertyBagPropertyDesc> PropertyDescs;
+	PropertyDescs.Reserve(CallFunctionInfo.FunctionHeader.Arguments.Num());
+	for (const FRigVMGraphFunctionArgument& Argument : CallFunctionInfo.FunctionHeader.Arguments)
+	{
+		const bool bIsBindable = Argument.Direction == ERigVMPinDirection::Input || Argument.Direction == ERigVMPinDirection::Output;
+		if (bIsBindable && !Argument.bIsInputVariable)
+		{
+			FString CppTypeString = Argument.CPPType.ToString();
+			UObject* CppTypeObject = Argument.CPPTypeObject.Get();
+			FPropertyBagPropertyDesc PropertyDesc = FRigVMMemoryStorageStruct::GeneratePropertyBagDescriptor(FRigVMPropertyDescription(Argument.Name, CppTypeString, CppTypeObject, Argument.DefaultValue));
+			if (Argument.Direction == ERigVMPinDirection::Output && Argument.CPPType == RigVMTypeUtils::BoolTypeName)
+			{
+				PropertyDesc.PropertyFlags &= ~CPF_Edit;	// Hide the result from the UI
+				ResultIndex = PropertyDescs.Num();
+			}
 
-	FAnimNextProgrammaticFunctionHeader AnimNextFunctionHeader = {};
-	AnimNextFunctionHeader.FunctionHeader = CallFunctionInfo.FunctionHeader;
-	OutCompileContext.AddUniqueFunctionHeader(AnimNextFunctionHeader);
+			PropertyDescs.Add(PropertyDesc);
+		}
+	}
+
+	// @TODO: Constrain functions we can pick to those returning only a bool
+	InstanceData->Parameters.AddProperties(PropertyDescs);
 }
 
 #endif // WITH_EDITOR
diff --git a/Engine/Plugins/Experimental/UAF/UAFStateTree/Source/UAFStateTree/Private/Tasks/AnimNextStateTreeRigVMTaskBase.cpp b/Engine/Plugins/Experimental/UAF/UAFStateTree/Source/UAFStateTree/Private/Tasks/AnimNextStateTreeRigVMTaskBase.cpp
--- a/Engine/Plugins/Experimental/UAF/UAFStateTree/Source/UAFStateTree/Private/Tasks/AnimNextStateTreeRigVMTaskBase.cpp
+++ b/Engine/Plugins/Experimental/UAF/UAFStateTree/Source/UAFStateTree/Private/Tasks/AnimNextStateTreeRigVMTaskBase.cpp
@@ -116,48 +116,101 @@ void FAnimNextStateTreeRigVMTaskBase::PostEditNodeChangeChainProperty(const FPro
 		// Function selected has changed. Update param struct & result index.
 		if (FAnimNextStateTreeRigVMTaskInstanceData* InstanceData = InstanceDataView.GetMutablePtr<FAnimNextStateTreeRigVMTaskInstanceData>())
 		{
-			InstanceData->Parameters.Reset();
-			ResultIndex = INDEX_NONE;
+			RegenerateInstanceDataFunctionParams(InstanceData);
+		}
+	}
+}
+
+void FAnimNextStateTreeRigVMTaskBase::GetProgrammaticFunctionHeaders(FAnimNextStateTreeProgrammaticFunctionHeaderParams& InProgrammaticFunctionHeaderParams, UStateTreeState* State, FStateTreeDataView InstanceDataView, const FStateTreeBindableStructDesc& Desc)
+{
+	using namespace UE::UAF::UncookedOnly;
+
+	FAnimNextGetFunctionHeaderCompileContext& OutCompileContext = InProgrammaticFunctionHeaderParams.OutCompileContext;
+	
+	StateName = State->Name;
+	NodeId = Desc.ID;
+
+	// Validate that our cached instance parameters match the latest RigVM signature. Update if needed.
+	UStateTree* OwningStateTree = State->GetTypedOuter<UStateTree>();
+	if (FAnimNextStateTreeRigVMTaskInstanceData* InstanceData = InstanceDataView.GetMutablePtr<FAnimNextStateTreeRigVMTaskInstanceData>())
+	{
+		auto IsBindingOutdated = [&]()
+		{
+			if (!InstanceData->Parameters.IsValid())
+			{
+				return false;
+			}
+
+			TConstArrayView<FPropertyBagPropertyDesc> CachedDescs = InstanceData->Parameters.GetPropertyBagStruct()->GetPropertyDescs();
+
+			// Our function header might be out of date. Get latest
+			CallFunctionInfo.FunctionHeader = FRigVMGraphFunctionHeader::FindGraphFunctionHeader(CallFunctionInfo.FunctionHeader.LibraryPointer);
+			int32 RigVMFunctionBindableArgumentIndex = 0;
 
-			TArray<FPropertyBagPropertyDesc> PropertyDescs;
-			PropertyDescs.Reserve(CallFunctionInfo.FunctionHeader.Arguments.Num());
 			for (const FRigVMGraphFunctionArgument& Argument : CallFunctionInfo.FunctionHeader.Arguments)
 			{
-				if (Argument.Direction == ERigVMPinDirection::Input || Argument.Direction == ERigVMPinDirection::Output)
+				const bool bIsBindable = Argument.Direction == ERigVMPinDirection::Input || Argument.Direction == ERigVMPinDirection::Output;
+				if (bIsBindable && !Argument.bIsInputVariable)
 				{
 					FString CppTypeString = Argument.CPPType.ToString();
 					UObject* CppTypeObject = Argument.CPPTypeObject.Get();
 					FPropertyBagPropertyDesc PropertyDesc = FRigVMMemoryStorageStruct::GeneratePropertyBagDescriptor(FRigVMPropertyDescription(Argument.Name, CppTypeString, CppTypeObject, Argument.DefaultValue));
-					if (Argument.Direction == ERigVMPinDirection::Output)
-					{
-						PropertyDesc.PropertyFlags &= ~CPF_Edit;	// Hide the results from the UI
-						ResultIndex = PropertyDescs.Num();
-					}
 
-					if (!Argument.bIsInputVariable)
+					if (!CachedDescs.IsValidIndex(RigVMFunctionBindableArgumentIndex) || CachedDescs[RigVMFunctionBindableArgumentIndex] != PropertyDesc)
 					{
-						PropertyDescs.Add(PropertyDesc);
+						return true;
 					}
+
+					RigVMFunctionBindableArgumentIndex++;
 				}
 			}
-		
-			InstanceData->Parameters.AddProperties(PropertyDescs);
+
+			const int32 NumBindableArguments = RigVMFunctionBindableArgumentIndex;
+			if (NumBindableArguments != CachedDescs.Num())
+			{
+				return true;
+			}
+
+			return false;
+		};
+
+		if (IsBindingOutdated())
+		{
+			RegenerateInstanceDataFunctionParams(InstanceData);
 		}
 	}
+
+	FAnimNextProgrammaticFunctionHeader AnimNextFunctionHeader = {};
+	AnimNextFunctionHeader.FunctionHeader = CallFunctionInfo.FunctionHeader;
+	OutCompileContext.AddUniqueFunctionHeader(AnimNextFunctionHeader);
 }
 
-void FAnimNextStateTreeRigVMTaskBase::GetProgrammaticFunctionHeaders(FAnimNextStateTreeProgrammaticFunctionHeaderParams& InProgrammaticFunctionHeaderParams, const UStateTreeState* State, const FStateTreeBindableStructDesc& Desc)
+void FAnimNextStateTreeRigVMTaskBase::RegenerateInstanceDataFunctionParams(TNotNull<FInstanceDataType*> InstanceData)
 {
-	using namespace UE::UAF::UncookedOnly;
+	InstanceData->Parameters.Reset();
+	ResultIndex = INDEX_NONE;
 
-	FAnimNextGetFunctionHeaderCompileContext& OutCompileContext = InProgrammaticFunctionHeaderParams.OutCompileContext;
-	
-	StateName = State->Name;
-	NodeId = Desc.ID;
+	TArray<FPropertyBagPropertyDesc> PropertyDescs;
+	PropertyDescs.Reserve(CallFunctionInfo.FunctionHeader.Arguments.Num());
+	for (const FRigVMGraphFunctionArgument& Argument : CallFunctionInfo.FunctionHeader.Arguments)
+	{
+		const bool bIsBindable = Argument.Direction == ERigVMPinDirection::Input || Argument.Direction == ERigVMPinDirection::Output;
+		if (bIsBindable && !Argument.bIsInputVariable)
+		{
+			FString CppTypeString = Argument.CPPType.ToString();
+			UObject* CppTypeObject = Argument.CPPTypeObject.Get();
+			FPropertyBagPropertyDesc PropertyDesc = FRigVMMemoryStorageStruct::GeneratePropertyBagDescriptor(FRigVMPropertyDescription(Argument.Name, CppTypeString, CppTypeObject, Argument.DefaultValue));
+			if (Argument.Direction == ERigVMPinDirection::Output)
+			{
+				PropertyDesc.PropertyFlags &= ~CPF_Edit;	// Hide the results from the UI
+				ResultIndex = PropertyDescs.Num();
+			}
 
-	FAnimNextProgrammaticFunctionHeader AnimNextFunctionHeader = {};
-	AnimNextFunctionHeader.FunctionHeader = CallFunctionInfo.FunctionHeader;
-	OutCompileContext.AddUniqueFunctionHeader(AnimNextFunctionHeader);
+			PropertyDescs.Add(PropertyDesc);
+		}
+	}
+
+	InstanceData->Parameters.AddProperties(PropertyDescs);
 }
 
 #endif // WITH_EDITOR
diff --git a/Engine/Plugins/Experimental/UAF/UAFStateTree/Source/UAFStateTreeUncookedOnly/Private/AnimNextStateTree_EditorData.cpp b/Engine/Plugins/Experimental/UAF/UAFStateTree/Source/UAFStateTreeUncookedOnly/Private/AnimNextStateTree_EditorData.cpp
--- a/Engine/Plugins/Experimental/UAF/UAFStateTree/Source/UAFStateTreeUncookedOnly/Private/AnimNextStateTree_EditorData.cpp
+++ b/Engine/Plugins/Experimental/UAF/UAFStateTree/Source/UAFStateTreeUncookedOnly/Private/AnimNextStateTree_EditorData.cpp
@@ -102,7 +102,7 @@ void UAnimNextStateTree_EditorData::BuildFunctionHeadersContext(const FRigVMComp
 					Desc.DataSource = EStateTreeBindableStructSource::Condition;
 
 					// Graph gen isn't const since we need to cache generated names in the node
-					const_cast<FAnimNextStateTreeRigVMConditionBase*>(Condition)->GetProgrammaticFunctionHeaders(ProgrammaticFunctionHeaderParams, &State, Desc);
+					const_cast<FAnimNextStateTreeRigVMConditionBase*>(Condition)->GetProgrammaticFunctionHeaders(ProgrammaticFunctionHeaderParams, &State, Node.GetInstance(), Desc);
 				}
 			}
 
@@ -120,7 +120,7 @@ void UAnimNextStateTree_EditorData::BuildFunctionHeadersContext(const FRigVMComp
 						Desc.DataSource = EStateTreeBindableStructSource::Transition;
 
 						// Graph gen isn't const since we need to cache generated names in the node
-						const_cast<FAnimNextStateTreeRigVMConditionBase*>(Condition)->GetProgrammaticFunctionHeaders(ProgrammaticFunctionHeaderParams, &State, Desc);
+						const_cast<FAnimNextStateTreeRigVMConditionBase*>(Condition)->GetProgrammaticFunctionHeaders(ProgrammaticFunctionHeaderParams, &State, Node.GetInstance(), Desc);
 					}
 				}
 			}
@@ -138,7 +138,7 @@ void UAnimNextStateTree_EditorData::BuildFunctionHeadersContext(const FRigVMComp
 					Desc.DataSource = EStateTreeBindableStructSource::Task;
 
 					// Graph gen isn't const since we need to cache generated names in the node
-					const_cast<FAnimNextStateTreeRigVMTaskBase*>(Condition)->GetProgrammaticFunctionHeaders(ProgrammaticFunctionHeaderParams, &State, Desc);
+					const_cast<FAnimNextStateTreeRigVMTaskBase*>(Condition)->GetProgrammaticFunctionHeaders(ProgrammaticFunctionHeaderParams, &State, Node.GetInstance(), Desc);
 				}
 			}
 
