Add support for default args to UAF ST RigVM functions
- Will now set default args when initially selecting function
- Will now recompile when default args are out of sync
- Will not overwrite existing values that have been modified from defaults.

#tests PIE, Editor

#Synced-CL 50972851

diff --git a/Engine/Plugins/Experimental/UAF/UAFStateTree/Source/UAFStateTree/Internal/Conditions/AnimNextStateTreeRigVMConditionBase.h b/Engine/Plugins/Experimental/UAF/UAFStateTree/Source/UAFStateTree/Internal/Conditions/AnimNextStateTreeRigVMConditionBase.h
--- a/Engine/Plugins/Experimental/UAF/UAFStateTree/Source/UAFStateTree/Internal/Conditions/AnimNextStateTreeRigVMConditionBase.h
+++ b/Engine/Plugins/Experimental/UAF/UAFStateTree/Source/UAFStateTree/Internal/Conditions/AnimNextStateTreeRigVMConditionBase.h
@@ -81,6 +81,10 @@ public:
 
 	UPROPERTY(meta = (DeprecatedProperty, DeprecationMessage = "InternalEventName is deprecated, use CallFunctionInfo.FunctionEvent"))
 	FName InternalEventName = NAME_None;
+
+	/** Bag to hold defaults for function signature. Used to perform manual delta serialization for defaults as InstancedPropertyBag doesn't support that */
+	UPROPERTY()
+	FInstancedPropertyBag DefaultValues;
 #endif
 
 	/** Info needed to call a UAF function. Also has a custom function picker. */
diff --git a/Engine/Plugins/Experimental/UAF/UAFStateTree/Source/UAFStateTree/Internal/Tasks/AnimNextStateTreeRigVMTaskBase.h b/Engine/Plugins/Experimental/UAF/UAFStateTree/Source/UAFStateTree/Internal/Tasks/AnimNextStateTreeRigVMTaskBase.h
--- a/Engine/Plugins/Experimental/UAF/UAFStateTree/Source/UAFStateTree/Internal/Tasks/AnimNextStateTreeRigVMTaskBase.h
+++ b/Engine/Plugins/Experimental/UAF/UAFStateTree/Source/UAFStateTree/Internal/Tasks/AnimNextStateTreeRigVMTaskBase.h
@@ -88,6 +88,10 @@ public:
 	
 	UPROPERTY(meta = (DeprecatedProperty, DeprecationMessage = "InternalEventName is deprecated, use CallFunctionInfo.FunctionEvent"))
 	FName InternalEventName = NAME_None;
+
+	/** Bag to hold defaults for function signature. Used to perform manual delta serialization for defaults as InstancedPropertyBag doesn't support that */
+	UPROPERTY()
+	FInstancedPropertyBag DefaultValues;
 #endif
 
 	/** Info needed to call a UAF function. Also has a custom function picker. */
diff --git a/Engine/Plugins/Experimental/UAF/UAFStateTree/Source/UAFStateTree/Private/Conditions/AnimNextStateTreeRigVMConditionBase.cpp b/Engine/Plugins/Experimental/UAF/UAFStateTree/Source/UAFStateTree/Private/Conditions/AnimNextStateTreeRigVMConditionBase.cpp
--- a/Engine/Plugins/Experimental/UAF/UAFStateTree/Source/UAFStateTree/Private/Conditions/AnimNextStateTreeRigVMConditionBase.cpp
+++ b/Engine/Plugins/Experimental/UAF/UAFStateTree/Source/UAFStateTree/Private/Conditions/AnimNextStateTreeRigVMConditionBase.cpp
@@ -17,6 +17,7 @@
 #include "RigVMFunctions/RigVMFunctionDefines.h"
 #include "StateTree.h"
 #include "StateTreeEditorData.h"
+#include "UAFStateTreeRigVMUtils.h"
 #include "UncookedOnlyUtils.h"
 #endif // WITH_EDITOR
 
@@ -110,47 +111,7 @@ void FAnimNextStateTreeRigVMConditionBase::GetProgrammaticFunctionHeaders(FAnimN
 	UStateTree* OwningStateTree = State->GetTypedOuter<UStateTree>();
 	if (FAnimNextStateTreeRigVMConditionInstanceData* InstanceData = InstanceDataView.GetMutablePtr<FAnimNextStateTreeRigVMConditionInstanceData>())
 	{
-		auto IsBindingOutdated = [&]()
-		{
-			if (!InstanceData->Parameters.IsValid())
-			{
-				return false;
-			}
-
-			TConstArrayView<FPropertyBagPropertyDesc> CachedDescs = InstanceData->Parameters.GetPropertyBagStruct()->GetPropertyDescs();
-
-			// Our function header might be out of date. Get latest
-			CallFunctionInfo.FunctionHeader = FRigVMGraphFunctionHeader::FindGraphFunctionHeader(CallFunctionInfo.FunctionHeader.LibraryPointer);
-			int32 RigVMFunctionBindableArgumentIndex = 0;
-
-			for (const FRigVMGraphFunctionArgument& Argument : CallFunctionInfo.FunctionHeader.Arguments)
-			{
-				const bool bIsBindable = Argument.Direction == ERigVMPinDirection::Input || Argument.Direction == ERigVMPinDirection::Output;
-				if (bIsBindable && !Argument.bIsInputVariable)
-				{
-					FString CppTypeString = Argument.CPPType.ToString();
-					UObject* CppTypeObject = Argument.CPPTypeObject.Get();
-					FPropertyBagPropertyDesc PropertyDesc = FRigVMMemoryStorageStruct::GeneratePropertyBagDescriptor(FRigVMPropertyDescription(Argument.Name, CppTypeString, CppTypeObject, Argument.DefaultValue));
-
-					if (!CachedDescs.IsValidIndex(RigVMFunctionBindableArgumentIndex) || CachedDescs[RigVMFunctionBindableArgumentIndex] != PropertyDesc)
-					{
-						return true;
-					}
-
-					RigVMFunctionBindableArgumentIndex++;
-				}
-			}
-
-			const int32 NumBindableArguments = RigVMFunctionBindableArgumentIndex;
-			if (NumBindableArguments != CachedDescs.Num())
-			{
-				return true;
-			}
-
-			return false;
-		};
-
-		if (IsBindingOutdated())
+		if (UE::UAF::StateTree::FUtils::IsBindingOutDated(CallFunctionInfo, DefaultValues, InstanceData->Parameters))
 		{
 			RegenerateInstanceDataFunctionParams(InstanceData);
 		}
@@ -163,31 +124,7 @@ void FAnimNextStateTreeRigVMConditionBase::GetProgrammaticFunctionHeaders(FAnimN
 
 void FAnimNextStateTreeRigVMConditionBase::RegenerateInstanceDataFunctionParams(TNotNull<FInstanceDataType*> InstanceData)
 {
-	InstanceData->Parameters.Reset();
-	ResultIndex = INDEX_NONE;
-
-	TArray<FPropertyBagPropertyDesc> PropertyDescs;
-	PropertyDescs.Reserve(CallFunctionInfo.FunctionHeader.Arguments.Num());
-	for (const FRigVMGraphFunctionArgument& Argument : CallFunctionInfo.FunctionHeader.Arguments)
-	{
-		const bool bIsBindable = Argument.Direction == ERigVMPinDirection::Input || Argument.Direction == ERigVMPinDirection::Output;
-		if (bIsBindable && !Argument.bIsInputVariable)
-		{
-			FString CppTypeString = Argument.CPPType.ToString();
-			UObject* CppTypeObject = Argument.CPPTypeObject.Get();
-			FPropertyBagPropertyDesc PropertyDesc = FRigVMMemoryStorageStruct::GeneratePropertyBagDescriptor(FRigVMPropertyDescription(Argument.Name, CppTypeString, CppTypeObject, Argument.DefaultValue));
-			if (Argument.Direction == ERigVMPinDirection::Output && Argument.CPPType == RigVMTypeUtils::BoolTypeName)
-			{
-				PropertyDesc.PropertyFlags &= ~CPF_Edit;	// Hide the result from the UI
-				ResultIndex = PropertyDescs.Num();
-			}
-
-			PropertyDescs.Add(PropertyDesc);
-		}
-	}
-
-	// @TODO: Constrain functions we can pick to those returning only a bool
-	InstanceData->Parameters.AddProperties(PropertyDescs);
+	UE::UAF::StateTree::FUtils::RegenerateParameterPropertyBag(CallFunctionInfo, DefaultValues, InstanceData->Parameters, ResultIndex);
 }
 
 #endif // WITH_EDITOR
diff --git a/Engine/Plugins/Experimental/UAF/UAFStateTree/Source/UAFStateTree/Private/Tasks/AnimNextStateTreeRigVMTaskBase.cpp b/Engine/Plugins/Experimental/UAF/UAFStateTree/Source/UAFStateTree/Private/Tasks/AnimNextStateTreeRigVMTaskBase.cpp
--- a/Engine/Plugins/Experimental/UAF/UAFStateTree/Source/UAFStateTree/Private/Tasks/AnimNextStateTreeRigVMTaskBase.cpp
+++ b/Engine/Plugins/Experimental/UAF/UAFStateTree/Source/UAFStateTree/Private/Tasks/AnimNextStateTreeRigVMTaskBase.cpp
@@ -18,6 +18,7 @@
 #include "RigVMFunctions/RigVMFunctionDefines.h"
 #include "StateTree.h"
 #include "StateTreeEditorData.h"
+#include "UAFStateTreeRigVMUtils.h"
 #include "UncookedOnlyUtils.h"
 #endif // WITH_EDITOR
 
@@ -134,47 +135,7 @@ void FAnimNextStateTreeRigVMTaskBase::GetProgrammaticFunctionHeaders(FAnimNextSt
 	UStateTree* OwningStateTree = State->GetTypedOuter<UStateTree>();
 	if (FAnimNextStateTreeRigVMTaskInstanceData* InstanceData = InstanceDataView.GetMutablePtr<FAnimNextStateTreeRigVMTaskInstanceData>())
 	{
-		auto IsBindingOutdated = [&]()
-		{
-			if (!InstanceData->Parameters.IsValid())
-			{
-				return false;
-			}
-
-			TConstArrayView<FPropertyBagPropertyDesc> CachedDescs = InstanceData->Parameters.GetPropertyBagStruct()->GetPropertyDescs();
-
-			// Our function header might be out of date. Get latest
-			CallFunctionInfo.FunctionHeader = FRigVMGraphFunctionHeader::FindGraphFunctionHeader(CallFunctionInfo.FunctionHeader.LibraryPointer);
-			int32 RigVMFunctionBindableArgumentIndex = 0;
-
-			for (const FRigVMGraphFunctionArgument& Argument : CallFunctionInfo.FunctionHeader.Arguments)
-			{
-				const bool bIsBindable = Argument.Direction == ERigVMPinDirection::Input || Argument.Direction == ERigVMPinDirection::Output;
-				if (bIsBindable && !Argument.bIsInputVariable)
-				{
-					FString CppTypeString = Argument.CPPType.ToString();
-					UObject* CppTypeObject = Argument.CPPTypeObject.Get();
-					FPropertyBagPropertyDesc PropertyDesc = FRigVMMemoryStorageStruct::GeneratePropertyBagDescriptor(FRigVMPropertyDescription(Argument.Name, CppTypeString, CppTypeObject, Argument.DefaultValue));
-
-					if (!CachedDescs.IsValidIndex(RigVMFunctionBindableArgumentIndex) || CachedDescs[RigVMFunctionBindableArgumentIndex] != PropertyDesc)
-					{
-						return true;
-					}
-
-					RigVMFunctionBindableArgumentIndex++;
-				}
-			}
-
-			const int32 NumBindableArguments = RigVMFunctionBindableArgumentIndex;
-			if (NumBindableArguments != CachedDescs.Num())
-			{
-				return true;
-			}
-
-			return false;
-		};
-
-		if (IsBindingOutdated())
+		if (UE::UAF::StateTree::FUtils::IsBindingOutDated(CallFunctionInfo, DefaultValues, InstanceData->Parameters))
 		{
 			RegenerateInstanceDataFunctionParams(InstanceData);
 		}
@@ -187,30 +148,7 @@ void FAnimNextStateTreeRigVMTaskBase::GetProgrammaticFunctionHeaders(FAnimNextSt
 
 void FAnimNextStateTreeRigVMTaskBase::RegenerateInstanceDataFunctionParams(TNotNull<FInstanceDataType*> InstanceData)
 {
-	InstanceData->Parameters.Reset();
-	ResultIndex = INDEX_NONE;
-
-	TArray<FPropertyBagPropertyDesc> PropertyDescs;
-	PropertyDescs.Reserve(CallFunctionInfo.FunctionHeader.Arguments.Num());
-	for (const FRigVMGraphFunctionArgument& Argument : CallFunctionInfo.FunctionHeader.Arguments)
-	{
-		const bool bIsBindable = Argument.Direction == ERigVMPinDirection::Input || Argument.Direction == ERigVMPinDirection::Output;
-		if (bIsBindable && !Argument.bIsInputVariable)
-		{
-			FString CppTypeString = Argument.CPPType.ToString();
-			UObject* CppTypeObject = Argument.CPPTypeObject.Get();
-			FPropertyBagPropertyDesc PropertyDesc = FRigVMMemoryStorageStruct::GeneratePropertyBagDescriptor(FRigVMPropertyDescription(Argument.Name, CppTypeString, CppTypeObject, Argument.DefaultValue));
-			if (Argument.Direction == ERigVMPinDirection::Output)
-			{
-				PropertyDesc.PropertyFlags &= ~CPF_Edit;	// Hide the results from the UI
-				ResultIndex = PropertyDescs.Num();
-			}
-
-			PropertyDescs.Add(PropertyDesc);
-		}
-	}
-
-	InstanceData->Parameters.AddProperties(PropertyDescs);
+	UE::UAF::StateTree::FUtils::RegenerateParameterPropertyBag(CallFunctionInfo, DefaultValues, InstanceData->Parameters, ResultIndex);
 }
 
 #endif // WITH_EDITOR
diff --git a/Engine/Plugins/Experimental/UAF/UAFStateTree/Source/UAFStateTree/Private/UAFStateTreeRigVMUtils.cpp b/Engine/Plugins/Experimental/UAF/UAFStateTree/Source/UAFStateTree/Private/UAFStateTreeRigVMUtils.cpp
new file mode 100644
--- /dev/null
+++ b/Engine/Plugins/Experimental/UAF/UAFStateTree/Source/UAFStateTree/Private/UAFStateTreeRigVMUtils.cpp
@@ -0,0 +1,148 @@
+// Copyright Epic Games, Inc. All Rights Reserved.
+
+#include "UAFStateTreeRigVMUtils.h"
+
+#if WITH_EDITOR
+
+#include "AnimNextFunctionHandle.h"
+#include "AnimNextStateTreeContext.h"
+#include "AnimNextStateTreeEditorOnlyTypes.h"
+#include "AnimNextRigVMAsset.h"
+#include "Graph/AnimNextGraphContextData.h"
+#include "Graph/AnimNextGraphInstance.h"
+#include "StructUtils/PropertyBag.h"
+#include "TraitCore/ExecutionContext.h"
+#include "RigVMBlueprintGeneratedClass.h"
+#include "RigVMCore/RigVMGraphFunctionDefinition.h"
+#include "RigVMFunctions/RigVMFunctionDefines.h"
+#include "Script/UAFRigVMComponent.h"
+#include "StateTree.h"
+#include "StateTreeEditorData.h"
+#include "StateTreeReference.h"
+#include "Traits/CallFunction.h"
+#include "UncookedOnlyUtils.h"
+#include "Variables/AnimNextSharedVariables.h"
+
+#endif // WITH_EDITOR
+
+namespace UE::UAF::StateTree
+{
+
+#if WITH_EDITOR
+
+	bool FUtils::IsBindingOutDated(FUAFCallFunctionInfo& CallFunctionInfo, FInstancedPropertyBag& DefaultValues, FInstancedPropertyBag& Parameters)
+	{
+		if (!Parameters.IsValid())
+		{
+			return false;
+		}
+
+		TConstArrayView<FPropertyBagPropertyDesc> CachedDescs = Parameters.GetPropertyBagStruct()->GetPropertyDescs();
+
+		// Our function header might be out of date. Get latest
+		CallFunctionInfo.FunctionHeader = FRigVMGraphFunctionHeader::FindGraphFunctionHeader(CallFunctionInfo.FunctionHeader.LibraryPointer);
+		int32 RigVMFunctionBindableArgumentIndex = 0;
+
+		for (const FRigVMGraphFunctionArgument& Argument : CallFunctionInfo.FunctionHeader.Arguments)
+		{
+			const bool bIsBindable = Argument.Direction == ERigVMPinDirection::Input || Argument.Direction == ERigVMPinDirection::Output;
+			if (bIsBindable && !Argument.bIsInputVariable)
+			{
+				FString CppTypeString = Argument.CPPType.ToString();
+				UObject* CppTypeObject = Argument.CPPTypeObject.Get();
+				FPropertyBagPropertyDesc PropertyDesc = FRigVMMemoryStorageStruct::GeneratePropertyBagDescriptor(FRigVMPropertyDescription(Argument.Name, CppTypeString, CppTypeObject, Argument.DefaultValue));
+
+				if (!CachedDescs.IsValidIndex(RigVMFunctionBindableArgumentIndex) || CachedDescs[RigVMFunctionBindableArgumentIndex] != PropertyDesc)
+				{
+					return true;
+				}
+
+				TValueOrError<FString, EPropertyBagResult> CachedDefaultValue = DefaultValues.GetValueSerializedString(Argument.Name);
+				if (CachedDefaultValue.IsValid() && CachedDefaultValue.GetValue() != Argument.DefaultValue)
+				{
+					return true;
+				}
+
+				RigVMFunctionBindableArgumentIndex++;
+			}
+		}
+
+		const int32 NumBindableArguments = RigVMFunctionBindableArgumentIndex;
+		if (NumBindableArguments != CachedDescs.Num())
+		{
+			return true;
+		}
+
+		return false;
+	}
+
+	void FUtils::RegenerateParameterPropertyBag(FUAFCallFunctionInfo& CallFunctionInfo, FInstancedPropertyBag& DefaultValues, FInstancedPropertyBag& Parameters, int32& ResultIndex)
+	{
+		// Default values may not have been captured (Ex: old content). Always propagate values in this case.
+		bool bDefaultValuesCached = DefaultValues.IsValid();
+
+		// Before we regenerate, get a list of GUIDs for changed values. We will  migrate these over current defaults
+		TArray<FGuid, TInlineAllocator<4>> ChangedValuePropertyIDs = {};
+		ChangedValuePropertyIDs.Reserve(DefaultValues.GetNumPropertiesInBag());
+		if (bDefaultValuesCached && ensure(DefaultValues.GetNumPropertiesInBag() == Parameters.GetNumPropertiesInBag()))
+		{
+			TConstArrayView<FPropertyBagPropertyDesc> CachedDescs = Parameters.GetPropertyBagStruct()->GetPropertyDescs();
+			for (int32 CachedDescIndex = 0; CachedDescIndex < CachedDescs.Num(); CachedDescIndex++)
+			{
+				FName CachedPropertyName = CachedDescs[CachedDescIndex].Name;
+				TValueOrError<FString, EPropertyBagResult> CachedValue = Parameters.GetValueSerializedString(CachedPropertyName);
+				TValueOrError<FString, EPropertyBagResult> DefaultValue = DefaultValues.GetValueSerializedString(CachedPropertyName);
+
+				if (ensure(CachedValue.IsValid() && DefaultValue.IsValid()))
+				{
+					if (CachedValue.GetValue() != DefaultValue.GetValue())
+					{
+						ChangedValuePropertyIDs.Add(CachedDescs[CachedDescIndex].ID);
+					}
+				}
+			}
+		}
+
+		// Update defaults, then use that + above GUIDs to selectively update parameters for default values
+		DefaultValues.Reset();
+
+		TArray<FPropertyBagPropertyDesc, TInlineAllocator<4>> PropertyDescs;
+		TArray<FStringView, TInlineAllocator<4>> DefaultValueStrings;
+		PropertyDescs.Reserve(CallFunctionInfo.FunctionHeader.Arguments.Num());
+		for (const FRigVMGraphFunctionArgument& Argument : CallFunctionInfo.FunctionHeader.Arguments)
+		{
+			const bool bIsBindable = Argument.Direction == ERigVMPinDirection::Input || Argument.Direction == ERigVMPinDirection::Output;
+			if (bIsBindable && !Argument.bIsInputVariable)
+			{
+				FString CppTypeString = Argument.CPPType.ToString();
+				UObject* CppTypeObject = Argument.CPPTypeObject.Get();
+				FPropertyBagPropertyDesc PropertyDesc = FRigVMMemoryStorageStruct::GeneratePropertyBagDescriptor(FRigVMPropertyDescription(Argument.Name, CppTypeString, CppTypeObject, Argument.DefaultValue));
+				if (Argument.Direction == ERigVMPinDirection::Output)
+				{
+					PropertyDesc.PropertyFlags &= ~CPF_Edit;	// Hide the results from the UI
+				}
+
+				PropertyDescs.Add(PropertyDesc);
+				DefaultValueStrings.Add(Argument.DefaultValue);
+			}
+		}
+
+		DefaultValues.AddProperties(PropertyDescs);
+		for (int32 PropertyIndex = 0; PropertyIndex < PropertyDescs.Num(); PropertyIndex++)
+		{
+			DefaultValues.SetValueSerializedString(PropertyDescs[PropertyIndex].Name, DefaultValueStrings[PropertyIndex].GetData());
+		}
+
+		if (bDefaultValuesCached)
+		{
+			Parameters.MigrateToNewBagInstanceWithOverrides(DefaultValues, ChangedValuePropertyIDs);
+		}
+		else
+		{
+			Parameters.MigrateToNewBagInstance(DefaultValues);
+		}
+	}
+
+#endif // WITH_EDITOR
+
+}
\ No newline at end of file
diff --git a/Engine/Plugins/Experimental/UAF/UAFStateTree/Source/UAFStateTree/Private/UAFStateTreeRigVMUtils.h b/Engine/Plugins/Experimental/UAF/UAFStateTree/Source/UAFStateTree/Private/UAFStateTreeRigVMUtils.h
new file mode 100644
--- /dev/null
+++ b/Engine/Plugins/Experimental/UAF/UAFStateTree/Source/UAFStateTree/Private/UAFStateTreeRigVMUtils.h
@@ -0,0 +1,22 @@
+// Copyright Epic Games, Inc. All Rights Reserved.
+
+#pragma once
+
+#include "CoreMinimal.h"
+
+struct FInstancedPropertyBag;
+struct FUAFCallFunctionInfo;
+
+namespace UE::UAF::StateTree
+{
+	struct FUtils
+	{
+#if WITH_EDITOR
+		/** True if parameters or default values no longer match the given function header. */
+		static bool IsBindingOutDated(FUAFCallFunctionInfo& CallFunctionInfo, FInstancedPropertyBag& DefaultValues, FInstancedPropertyBag& Parameters);
+
+		/** Repopulates parameters and default values with the given function header. */
+		static void RegenerateParameterPropertyBag(FUAFCallFunctionInfo& CallFunctionInfo, FInstancedPropertyBag& DefaultValues, FInstancedPropertyBag& Parameters, int32& ResultIndex);
+#endif // WITH_EDITOR
+	};
+}
\ No newline at end of file
