Add ability to call a function on state exit

#tests Editor

#Synced-CL 50972851

diff --git a/Engine/Plugins/Experimental/UAF/UAFStateTree/Source/UAFStateTree/Internal/Tasks/AnimNextStateTreeRigVMTaskBase.h b/Engine/Plugins/Experimental/UAF/UAFStateTree/Source/UAFStateTree/Internal/Tasks/AnimNextStateTreeRigVMTaskBase.h
--- a/Engine/Plugins/Experimental/UAF/UAFStateTree/Source/UAFStateTree/Internal/Tasks/AnimNextStateTreeRigVMTaskBase.h
+++ b/Engine/Plugins/Experimental/UAF/UAFStateTree/Source/UAFStateTree/Internal/Tasks/AnimNextStateTreeRigVMTaskBase.h
@@ -27,7 +27,8 @@ UENUM()
 enum class EStateTreeFunctionCallFrequency
 {
 	OnEnter,
-	OnTick
+	OnTick,
+	OnExit,
 };
 
 USTRUCT()
@@ -58,6 +59,7 @@ public:
 	//~ Begin FStateTreeTaskBase Interface
 	virtual bool Link(FStateTreeLinker& Linker) override;
 	virtual EStateTreeRunStatus EnterState(FStateTreeExecutionContext& Context, const FStateTreeTransitionResult& Transition) const override;
+	virtual void ExitState(FStateTreeExecutionContext& Context, const FStateTreeTransitionResult& Transition) const override;
 	virtual EStateTreeRunStatus Tick(FStateTreeExecutionContext& Context, const float DeltaTime) const override;
 	//~ End FStateTreeTaskBase Interface
 
diff --git a/Engine/Plugins/Experimental/UAF/UAFStateTree/Source/UAFStateTree/Private/Tasks/AnimNextStateTreeRigVMTaskBase.cpp b/Engine/Plugins/Experimental/UAF/UAFStateTree/Source/UAFStateTree/Private/Tasks/AnimNextStateTreeRigVMTaskBase.cpp
--- a/Engine/Plugins/Experimental/UAF/UAFStateTree/Source/UAFStateTree/Private/Tasks/AnimNextStateTreeRigVMTaskBase.cpp
+++ b/Engine/Plugins/Experimental/UAF/UAFStateTree/Source/UAFStateTree/Private/Tasks/AnimNextStateTreeRigVMTaskBase.cpp
@@ -45,6 +45,16 @@ EStateTreeRunStatus FAnimNextStateTreeRigVMTaskBase::EnterState(FStateTreeExecut
 	return CallFunction(Context);
 }
 
+void FAnimNextStateTreeRigVMTaskBase::ExitState(FStateTreeExecutionContext& Context, const FStateTreeTransitionResult& Transition) const
+{
+	if (CallFunctionFrequency != EStateTreeFunctionCallFrequency::OnExit)
+	{
+		return Super::ExitState(Context, Transition);
+	}
+
+	CallFunction(Context);
+}
+
 EStateTreeRunStatus FAnimNextStateTreeRigVMTaskBase::Tick(FStateTreeExecutionContext& Context, const float DeltaTime) const
 {
 	if (CallFunctionFrequency != EStateTreeFunctionCallFrequency::OnTick)
