Initial StateTree Anim Node integration
- refactored (and renamed) state tree context into a base, and two separate implementations for Traits/Nodes
- Added anim node which can play a state tree, and a couple of temporary custom conditions which read variables from a variable reference

#Synced-CL 51390015

#Synced-CL 51390015

diff --git a/Engine/Plugins/Experimental/UAF/UAFStateTree/Source/UAFStateTree/Internal/AnimNode/UAFStateTreeNode.h b/Engine/Plugins/Experimental/UAF/UAFStateTree/Source/UAFStateTree/Internal/AnimNode/UAFStateTreeNode.h
new file mode 100644
--- /dev/null
+++ b/Engine/Plugins/Experimental/UAF/UAFStateTree/Source/UAFStateTree/Internal/AnimNode/UAFStateTreeNode.h
@@ -0,0 +1,77 @@
+// Copyright Epic Games, Inc. All Rights Reserved.
+
+#pragma once
+
+#include "CoreMinimal.h"
+
+#include "AnimNextStateTree.h"
+#include "StateTreeConditionBase.h"
+#include "StateTreeExecutionContext.h"
+#include "StateTreeLinker.h"
+
+#include "UAF/UAFAssetData.h"
+#include "UAF/AnimNodeCore/UAFAnimNodeData.h"
+#include "UAF/AnimNodeCore/UAFAnimNode.h"
+#include "UAF/AnimNodeCore/UAFTransitionNodeData.h"
+#include "UAF/AnimNodes/IUAFAnimNodeDataHasAsset.h"
+
+#include "UAFStateTreeNode.generated.h"
+
+namespace UE::UAF::StateTree
+{
+	
+struct FUAFStateTreeNodeData;
+
+// Anim Node for playing a State tree
+class FUAFStateTreeNode : public FUAFAnimNode
+{
+public:
+	FUAFStateTreeNode(FUAFAnimGraphUpdateContext& Context, TObjectPtr<const UStateTree> StateTree);
+
+	UAFSTATETREE_API virtual void PreUpdate(FUAFAnimGraphUpdateContext& Context) override;
+	UAFSTATETREE_API virtual void* GetInterface(FUAFAnimNodeInterfaceId Id) override;
+	UAFSTATETREE_API static void AddReferencedObjects(FUAFAnimNode* This, FReferenceCollector& Collector);
+
+	UAFSTATETREE_API void BlendTo(FUAFAnimGraphUpdateContext& Context, TConstStructView<FUAFGraphFactoryAsset> AssetData, FUAFTransitionNodeData* TransitionBlend = nullptr);
+
+private:
+
+	TObjectPtr<const UStateTree> StateTree = nullptr;
+	FStateTreeInstanceData InstanceData;
+	FStateTreeExecutionContext::FExternalGlobalParameters StateTreeExternalParameters;
+};
+
+USTRUCT(DisplayName = "State Tree")
+struct FUAFStateTreeNodeData : public FUAFAnimNodeData
+#if CPP
+, public IUAFAnimNodeDataHasAsset
+#endif
+{
+	GENERATED_BODY()
+	
+	UAFSTATETREE_API virtual FUAFAnimNodePtr CreateInstance(FUAFAnimGraphUpdateContext& Context) const override;
+	UAFSTATETREE_API virtual void* GetInterface(FUAFAnimNodeInterfaceId Id) override;
+
+	// HasAsset interface
+	UAFSTATETREE_API virtual UObject* GetAsset() const override
+	{
+		return StateTree;
+	}
+
+	UAFSTATETREE_API const UStateTree* GetStateTree() const
+	{
+		return StateTree ? StateTree->StateTree : nullptr;
+	}
+
+	UAFSTATETREE_API void SetStateTree(UAnimNextStateTree* InStateTree)
+	{
+		StateTree = InStateTree;
+	}
+
+private:
+
+	UPROPERTY(EditAnywhere, Category = "Default")
+	TObjectPtr<UAnimNextStateTree> StateTree;
+};
+	
+}
\ No newline at end of file
diff --git a/Engine/Plugins/Experimental/UAF/UAFStateTree/Source/UAFStateTree/Internal/AnimNode/UAFStateTreeNodeConditions.h b/Engine/Plugins/Experimental/UAF/UAFStateTree/Source/UAFStateTree/Internal/AnimNode/UAFStateTreeNodeConditions.h
new file mode 100644
--- /dev/null
+++ b/Engine/Plugins/Experimental/UAF/UAFStateTree/Source/UAFStateTree/Internal/AnimNode/UAFStateTreeNodeConditions.h
@@ -0,0 +1,93 @@
+// Copyright Epic Games, Inc. All Rights Reserved.
+
+#pragma once
+
+#include "CoreMinimal.h"
+
+#include "AnimNextStateTree.h"
+#include "StateTreeConditionBase.h"
+#include "StateTreeExecutionContext.h"
+#include "StateTreeLinker.h"
+
+#include "UAF/UAFAssetData.h"
+#include "UAFStateTreeContext.h"
+
+#include "UAFStateTreeNodeConditions.generated.h"
+
+namespace UE::UAF::StateTree
+{
+	
+struct FUAFStateTreeNodeData;
+
+USTRUCT()
+struct FUAFDummyInstanceData
+{
+	GENERATED_BODY()
+};
+
+// These 2 State Tree conditions are temporary and only exist to test the state tree AnimNode implementation until variables can be better integrated
+
+USTRUCT(MinimalAPI, Meta = (Hidden))
+struct FUAFVariableConditionBase : public FStateTreeConditionBase
+{
+	GENERATED_BODY()
+
+	using FInstanceDataType = FUAFDummyInstanceData;
+
+	FUAFVariableConditionBase() = default;
+
+	virtual bool Link(FStateTreeLinker& Linker) override
+	{
+		Linker.LinkExternalData(ContextHandle);
+		return true;
+	}
+
+	virtual const UStruct* GetInstanceDataType() const override { return FInstanceDataType::StaticStruct(); }
+
+	TStateTreeExternalDataHandle<FUAFStateTreeContext> ContextHandle;	
+};
+
+UENUM()
+enum class EUAFFloatCompareConditionType
+{
+	Less,
+	Greater,
+};
+
+USTRUCT(MinimalAPI, DisplayName = "(PROTOTYPE) UAF Float Variable Compare")
+struct FUAFFloatCompareCondition : public FUAFVariableConditionBase
+{
+	GENERATED_BODY()
+
+	FUAFFloatCompareCondition() = default;
+
+	virtual bool TestCondition(FStateTreeExecutionContext& Context) const override;
+
+	UPROPERTY(EditAnywhere, Category = Default)
+	FAnimNextVariableReference Left;
+
+	UPROPERTY(EditAnywhere, Category = Default)
+	EUAFFloatCompareConditionType Compare;
+
+	UPROPERTY(EditAnywhere, Category = Default)
+	double Right;
+};
+
+
+USTRUCT(MinimalAPI, DisplayName = "(PROTOTYPE) UAF Enum Variable Compare")
+struct FUAFEnumCompareCondition : public FUAFVariableConditionBase
+{
+	GENERATED_BODY()
+
+	FUAFEnumCompareCondition() = default;
+
+	virtual bool TestCondition(FStateTreeExecutionContext& Context) const override;
+
+	UPROPERTY(EditAnywhere, Category = Default)
+	FAnimNextVariableReference Left;
+
+	UPROPERTY(EditAnywhere, Category = Default)
+	uint8 Right;
+};
+
+}
\ No newline at end of file
diff --git a/Engine/Plugins/Experimental/UAF/UAFStateTree/Source/UAFStateTree/Internal/Conditions/AnimNextStateTreeRigVMConditionBase.h b/Engine/Plugins/Experimental/UAF/UAFStateTree/Source/UAFStateTree/Internal/Conditions/AnimNextStateTreeRigVMConditionBase.h
--- a/Engine/Plugins/Experimental/UAF/UAFStateTree/Source/UAFStateTree/Internal/Conditions/AnimNextStateTreeRigVMConditionBase.h
+++ b/Engine/Plugins/Experimental/UAF/UAFStateTree/Source/UAFStateTree/Internal/Conditions/AnimNextStateTreeRigVMConditionBase.h
@@ -19,7 +19,7 @@ class UStateTreeState;
 struct FAnimNextStateTreeProgrammaticFunctionHeaderParams;
 struct FStateTreeBindableStructDesc;
 struct FStateTreeExecutionContext;
-struct FAnimNextStateTreeTraitContext;
+struct FUAFStateTreeContext;
 struct FRigVMClient;
 
 #define UE_API UAFSTATETREE_API
@@ -105,7 +105,7 @@ public:
 	UE_API void PostSerialize(const FArchive& Ar);
 
 public:
-	TStateTreeExternalDataHandle<FAnimNextStateTreeTraitContext> TraitContextHandle;
+	TStateTreeExternalDataHandle<FUAFStateTreeContext> TraitContextHandle;
 };
 
 #if WITH_EDITORONLY_DATA
diff --git a/Engine/Plugins/Experimental/UAF/UAFStateTree/Source/UAFStateTree/Internal/Tasks/AnimNextStateTreeGraphInstanceTask.h b/Engine/Plugins/Experimental/UAF/UAFStateTree/Source/UAFStateTree/Internal/Tasks/AnimNextStateTreeGraphInstanceTask.h
--- a/Engine/Plugins/Experimental/UAF/UAFStateTree/Source/UAFStateTree/Internal/Tasks/AnimNextStateTreeGraphInstanceTask.h
+++ b/Engine/Plugins/Experimental/UAF/UAFStateTree/Source/UAFStateTree/Internal/Tasks/AnimNextStateTreeGraphInstanceTask.h
@@ -5,6 +5,7 @@
 #include "AnimNextStateTreeTypes.h"
 #include "Graph/AnimNextAnimationGraph.h"
 #include <limits>
+
 #include "AlphaBlend.h"
 #include "Factory/AnimNextFactoryParams.h"
 #include "StructUtils/PropertyBag.h"
@@ -12,7 +13,7 @@
 
 #include "AnimNextStateTreeGraphInstanceTask.generated.h"
 
-struct FAnimNextStateTreeTraitContext;
+struct FUAFStateTreeContext;
 class UUAFBlendProfile;
 
 USTRUCT()
@@ -131,5 +132,5 @@ protected:
 	virtual void GetObjectReferences(TArray<const UObject*>& OutReferencedObjects, const FStateTreeDataView InstanceDataView) const override;
 #endif
 public:
-	TStateTreeExternalDataHandle<FAnimNextStateTreeTraitContext> TraitContextHandle;
+	TStateTreeExternalDataHandle<FUAFStateTreeContext> TraitContextHandle;
 };
diff --git a/Engine/Plugins/Experimental/UAF/UAFStateTree/Source/UAFStateTree/Internal/Tasks/AnimNextStateTreeRigVMTaskBase.h b/Engine/Plugins/Experimental/UAF/UAFStateTree/Source/UAFStateTree/Internal/Tasks/AnimNextStateTreeRigVMTaskBase.h
--- a/Engine/Plugins/Experimental/UAF/UAFStateTree/Source/UAFStateTree/Internal/Tasks/AnimNextStateTreeRigVMTaskBase.h
+++ b/Engine/Plugins/Experimental/UAF/UAFStateTree/Source/UAFStateTree/Internal/Tasks/AnimNextStateTreeRigVMTaskBase.h
@@ -19,7 +19,7 @@ struct FAnimNextStateTreeProgrammaticFunctionHeaderParams;
 struct FStateTreeBindableStructDesc;
 struct FStateTreeExecutionContext;
 struct FRigVMClient;
-struct FAnimNextStateTreeTraitContext;
+struct FUAFStateTreeContext;
 
 #define UE_API UAFSTATETREE_API
 
@@ -120,7 +120,7 @@ public:
 
 public:
 
-	TStateTreeExternalDataHandle<FAnimNextStateTreeTraitContext> TraitContextHandle;
+	TStateTreeExternalDataHandle<FUAFStateTreeContext> TraitContextHandle;
 
 protected:
 
diff --git a/Engine/Plugins/Experimental/UAF/UAFStateTree/Source/UAFStateTree/Private/AnimNextStateTreeContext.cpp b/Engine/Plugins/Experimental/UAF/UAFStateTree/Source/UAFStateTree/Private/AnimNextStateTreeContext.cpp
deleted file mode 100644
diff --git a/Engine/Plugins/Experimental/UAF/UAFStateTree/Source/UAFStateTree/Private/AnimNextStateTreeContext.h b/Engine/Plugins/Experimental/UAF/UAFStateTree/Source/UAFStateTree/Private/AnimNextStateTreeContext.h
deleted file mode 100644
diff --git a/Engine/Plugins/Experimental/UAF/UAFStateTree/Source/UAFStateTree/Private/AnimNextStateTreeModule.cpp b/Engine/Plugins/Experimental/UAF/UAFStateTree/Source/UAFStateTree/Private/AnimNextStateTreeModule.cpp
--- a/Engine/Plugins/Experimental/UAF/UAFStateTree/Source/UAFStateTree/Private/AnimNextStateTreeModule.cpp
+++ b/Engine/Plugins/Experimental/UAF/UAFStateTree/Source/UAFStateTree/Private/AnimNextStateTreeModule.cpp
@@ -2,18 +2,17 @@
 
 #include "AnimNextStateTreeModule.h"
 
+#include "AnimNextStateTree.h"
 #include "StateTree.h"
 #include "StateTreeReference.h"
+#include "AnimNode/UAFStateTreeNode.h"
 #include "Graph/AnimNextAnimationGraph.h"
 #include "RigVMCore/RigVMRegistry.h"
+#include "UAF/AnimNodeCore/UAFAnimNodeFactory.h"
 
 namespace UE::UAF::StateTree
 {
 
-void FAnimNextStateTreeModule::ShutdownModule()
-{
-}
-
 void FAnimNextStateTreeModule::StartupModule()
 {
 	static TPair<UClass*, FRigVMRegistry::ERegisterObjectOperation> const AllowedObjectTypes[] =
@@ -31,8 +30,20 @@ void FAnimNextStateTreeModule::StartupModule()
 	};
 
 	FRigVMRegistry::Get().RegisterStructTypes(AllowedStructTypes);
+
+	
+	UAFStateTreeClassPath = FUAFAnimNodeFactory::RegisterAsset(UAnimNextStateTree::StaticClass(), [](UObject* Object, FUAFAnimGraphUpdateContext& Context)
+		{
+			UAnimNextStateTree* StateTree = CastChecked<UAnimNextStateTree>(Object);
+			return MakeAnimNode<FUAFStateTreeNode>(Context, StateTree->StateTree);
+		});
+}
+	
+void FAnimNextStateTreeModule::ShutdownModule()
+{
+	FUAFAnimNodeFactory::UnregisterAsset(UAFStateTreeClassPath);
 }
 
 IMPLEMENT_MODULE(FAnimNextStateTreeModule, UAFStateTree)
 
-}
\ No newline at end of file
+}
diff --git a/Engine/Plugins/Experimental/UAF/UAFStateTree/Source/UAFStateTree/Private/AnimNextStateTreeModule.h b/Engine/Plugins/Experimental/UAF/UAFStateTree/Source/UAFStateTree/Private/AnimNextStateTreeModule.h
--- a/Engine/Plugins/Experimental/UAF/UAFStateTree/Source/UAFStateTree/Private/AnimNextStateTreeModule.h
+++ b/Engine/Plugins/Experimental/UAF/UAFStateTree/Source/UAFStateTree/Private/AnimNextStateTreeModule.h
@@ -3,6 +3,7 @@
 #pragma once
 
 #include "IAnimNextStateTreeModule.h"
+#include "UObject/TopLevelAssetPath.h"
 
 namespace UE::UAF::StateTree
 {
@@ -12,6 +13,9 @@ class FAnimNextStateTreeModule : public IAnimNextStateTreeModule
 public:
 	virtual void ShutdownModule() override;
 	virtual void StartupModule() override;
+	
+private:
+	FTopLevelAssetPath UAFStateTreeClassPath;
 };
 
-}
\ No newline at end of file
+}
diff --git a/Engine/Plugins/Experimental/UAF/UAFStateTree/Source/UAFStateTree/Private/AnimNextStateTreeSchema.cpp b/Engine/Plugins/Experimental/UAF/UAFStateTree/Source/UAFStateTree/Private/AnimNextStateTreeSchema.cpp
--- a/Engine/Plugins/Experimental/UAF/UAFStateTree/Source/UAFStateTree/Private/AnimNextStateTreeSchema.cpp
+++ b/Engine/Plugins/Experimental/UAF/UAFStateTree/Source/UAFStateTree/Private/AnimNextStateTreeSchema.cpp
@@ -11,13 +11,14 @@
 #include "Components/ActorComponent.h"
 #include "Subsystems/WorldSubsystem.h"
 #include "GameFramework/Actor.h"
-#include "AnimNextStateTreeContext.h"
+#include "UAFStateTreeContext.h"
+#include "AnimNode/UAFStateTreeNodeConditions.h"
 #include "Conditions/AnimNextStateTreeRigVMConditionBase.h"
 
 const FName UStateTreeAnimNextSchema::AnimStateTreeExecutionContextName("ExecutionContext");
 
 UStateTreeAnimNextSchema::UStateTreeAnimNextSchema() : ContextDataDescs(
-	{{ AnimStateTreeExecutionContextName, FAnimNextStateTreeTraitContext::StaticStruct(), FGuid(0xDFB93B7F, 0xEDBE4906, 0x851C66B2, 0x7585FA21) }})
+	{{ AnimStateTreeExecutionContextName, FUAFStateTreeContext::StaticStruct(), FGuid(0xDFB93B7F, 0xEDBE4906, 0x851C66B2, 0x7585FA21) }})
 {
 }
 
@@ -28,6 +29,7 @@ bool UStateTreeAnimNextSchema::IsStructAllowed(const UScriptStruct* InScriptStru
 	|| InScriptStruct->IsChildOf(FStateTreeTaskCommonBase::StaticStruct())
 	|| InScriptStruct->IsChildOf(FStateTreeConsiderationCommonBase::StaticStruct())
 	|| InScriptStruct->IsChildOf(FAnimNextStateTreeRigVMConditionBase::StaticStruct())
+	|| InScriptStruct->IsChildOf(UE::UAF::StateTree::FUAFVariableConditionBase::StaticStruct())
 	|| InScriptStruct->IsChildOf(FStateTreePropertyFunctionCommonBase::StaticStruct())
 	|| InScriptStruct->IsChildOf(FAnimNextStateTreeEvaluatorBase::StaticStruct())
 	|| InScriptStruct->IsChildOf(FAnimNextStateTreeTaskBase::StaticStruct());
@@ -40,7 +42,7 @@ bool UStateTreeAnimNextSchema::IsClassAllowed(const UClass* InClass) const
 
 bool UStateTreeAnimNextSchema::IsExternalItemAllowed(const UStruct& InStruct) const
 {
-	return InStruct.IsChildOf<FAnimNextStateTreeTraitContext>();
+	return InStruct.IsChildOf<FUAFStateTreeContext>();
 }
 
 EStateTreeParameterDataType UStateTreeAnimNextSchema::GetGlobalParameterDataType() const
diff --git a/Engine/Plugins/Experimental/UAF/UAFStateTree/Source/UAFStateTree/Private/AnimNode/UAFStateTreeNode.cpp b/Engine/Plugins/Experimental/UAF/UAFStateTree/Source/UAFStateTree/Private/AnimNode/UAFStateTreeNode.cpp
new file mode 100644
--- /dev/null
+++ b/Engine/Plugins/Experimental/UAF/UAFStateTree/Source/UAFStateTree/Private/AnimNode/UAFStateTreeNode.cpp
@@ -0,0 +1,298 @@
+// Copyright Epic Games, Inc. All Rights Reserved.
+
+#include "AnimNode/UAFStateTreeNode.h"
+
+#include "AnimNextStateTreeSchema.h"
+#include "StateTreeExecutionContext.h"
+#include "UAFStateTreeNodeContext.h"
+#include "Script/UAFRigVMComponent.h"
+#include "Logging/LogScopedVerbosityOverride.h"
+#include "UAF/AnimNodeCore/UAFAnimNodeUpdate.h"
+#include "UAF/AnimNodes/IUAFAnimNodeTimeline.h"
+#include "UAF/AnimNodes/UAFSimpleTransition.h"
+#include "UAF/AnimNodeCore/UAFGraphFactoryAssetAnimNodeFactory.h"
+#include "ObjectTrace.h"
+
+#if ENABLE_ANIM_DEBUG 
+#include "Debugger/StateTreeRuntimeValidation.h"
+#endif // ENABLE_ANIM_DEBUG
+
+namespace UE::UAF::StateTree
+{
+	
+FUAFStateTreeNode::FUAFStateTreeNode(FUAFAnimGraphUpdateContext& Context, TObjectPtr<const UStateTree> InStateTree)
+	: FUAFAnimNode(Context)
+{
+	InitializeAs<FUAFStateTreeNode>(Context);
+	
+	StateTree = InStateTree;
+
+	if (StateTree)
+	{
+		UObject* Owner = GetTransientPackage();
+#if ENABLE_ANIM_DEBUG 
+		Owner = const_cast<UObject*>(Context.GetHostObject());
+		// @TODO: Makes unique, but breaks visual logger
+		Owner = Owner ? Owner->GetOuter() : Owner;
+
+		if (!Owner)
+		{
+			return;
+		}
+#endif // ENABLE_ANIM_DEBUG
+
+		FStateTreeExecutionContext StateTreeExecutionContext(*Owner, *StateTree, InstanceData);
+	
+		FUAFStateTreeNodeContext UAFStateTreeContext(Context, this);
+		StateTreeExecutionContext.SetContextDataByName(UStateTreeAnimNextSchema::AnimStateTreeExecutionContextName, FStateTreeDataView(FUAFStateTreeContext::StaticStruct(), &UAFStateTreeContext));
+
+#if OBJECT_TRACE_ENABLED
+		// for now set the outer trace id to the owner object since I'm not tracing node instances yet
+		StateTreeExecutionContext.SetOuterTraceId(FObjectTrace::GetObjectId(Owner));
+#endif
+
+		const FInstancedPropertyBag& StateTreeParameters = StateTree->GetDefaultParameters();
+		StateTreeExternalParameters.Reset();
+	
+		TArray<int32> StartPropertyOffsets;
+		FUAFRigVMComponent& RigVMComponent = Context.GetVariablesOwner()->GetComponent<FUAFRigVMComponent>();
+		const FRigVMExtendedExecuteContext& ExtendedExecuteContext = RigVMComponent.GetExtendedExecuteContext();
+
+		// Variable Binding:  It might be possible to get this block of code working in a way where it binds the state tree variables to their values in the System/Module, but as-is the layouts mismatch and it fails.
+		// #if ENABLE_ANIM_DEBUG 
+		// if (!Context.GetVariablesOwner()->LayoutMatches(StateTree->GetDefaultParameters()))
+		// {
+		// 	UE_LOG(LogAnimation, Error, TEXT("Anim StateTree Parameter Layout Mismatch: %s - StateTree: %s"), *Owner->GetFName().ToString(),  *StateTree->GetOuter()->GetFName().ToString());
+		// 	StateTreeExecutionContext.Stop(EStateTreeRunStatus::Failed);
+		// 	return;
+		//}
+		// #endif // ENABLE_ANIM_DEBUG 
+		//
+		// const int32 NumVariables = StateTreeParameters.GetNumPropertiesInBag();
+		// for(int32 VariableIndex = 0; VariableIndex < NumVariables; ++VariableIndex)
+		// {
+		// 	const FPropertyBagPropertyDesc& Desc = StateTreeParameters.GetPropertyBagStruct()->GetPropertyDescs()[VariableIndex];
+		// 	StartPropertyOffsets.Add(Desc.CachedProperty->GetOffset_ForInternal());
+		// }
+		//
+		// for (const FPropertyBindingCopyInfoBatch& Batch : StateTree->GetPropertyBindings().GetCopyBatches())
+		// {
+		// 	for (const FPropertyBindingCopyInfo& Copy : StateTree->GetPropertyBindings().Super::GetBatchCopies(Batch))
+		// 	{
+		// 		const FStateTreeDataHandle& Handle = Copy.SourceDataHandle.Get<FStateTreeDataHandle>();
+		// 		if (Handle.GetSource() == EStateTreeDataSourceType::ExternalGlobalParameterData)
+		// 		{
+		//
+		// 			const int32 RequiredOffset = [&Copy]()
+		// 			{
+		// 				if(Copy.SourceIndirection.Type == EPropertyBindingPropertyAccessType::Offset)
+		// 				{
+		// 					return (int32)Copy.SourceIndirection.Offset;
+		// 				}
+		//
+		// 				checkf(false, TEXT("Only expecting offset indirections for remapping"));
+		// 				return Copy.SourceLeafProperty->GetOffset_ForInternal();
+		// 			}();
+		//
+		// 			uint8* MemoryPtr = nullptr;
+		// 			const int32 NumOffsets = StartPropertyOffsets.Num();
+		// 			for (int32 Index = 0; Index < NumOffsets; ++Index)
+		// 			{
+		// 				const int32 NextOffset = (Index < (NumOffsets - 1)) ? StartPropertyOffsets[Index + 1] : INDEX_NONE;
+		// 				if (RequiredOffset >= StartPropertyOffsets[Index] && (RequiredOffset < NextOffset || NextOffset == INDEX_NONE))
+		// 				{
+		// 					// Incorporate the UPropertyBag root-level property offset into the remapped memory-ptr itself, so that the property-access indirection works "as-normal"
+		// 					MemoryPtr = (uint8*)ExtendedExecuteContext.ExternalVariableRuntimeData[Index].Memory - StartPropertyOffsets[Index];
+		// 					break;
+		// 				}
+		// 			}
+		// 			check(MemoryPtr != nullptr);
+		//
+		// 			if (StateTreeExternalParameters.Add(Copy, MemoryPtr))
+		// 			{
+		// #if ENABLE_ANIM_DEBUG 
+		// 						UE_CLOG(CVarLogPropertyBindingMemoryPtrInfo.GetValueOnAnyThread(), LogAnimation, Warning, TEXT("Mapped: Source: %s\nTarget: %sSize: %i\nOffset: %i\nType: %s"), *Copy.SourceLeafProperty->GetName(), *Copy.TargetLeafProperty->GetName(), Copy.CopySize, RequiredOffset, *FindObject<UEnum>(nullptr, TEXT("/Script/PropertyBindingUtils.EPropertyCopyType"))->GetNameStringByValue(static_cast<int64>(Copy.Type)));
+		// 					}
+		// 					else
+		// 					{
+		// 						UE_CLOG(CVarLogPropertyBindingMemoryPtrInfo.GetValueOnAnyThread(), LogAnimation, Warning, TEXT("Skipped: Source: %s\nTarget: %sSize: %i\nOffset: %i\nType: %s"), *Copy.SourceLeafProperty->GetName(), *Copy.TargetLeafProperty->GetName(), Copy.CopySize, RequiredOffset, *FindObject<UEnum>(nullptr, TEXT("/Script/PropertyBindingUtils.EPropertyCopyType"))->GetNameStringByValue(static_cast<int64>(Copy.Type)));
+		// #endif // ENABLE_ANIM_DEBUG 
+		// 			}
+		// 		}
+		// 	}
+		// }
+	
+		// StateTreeExecutionContext.SetExternalGlobalParameters(&StateTreeExternalParameters);
+
+		StateTreeExecutionContext.SetCollectExternalDataCallback(FOnCollectStateTreeExternalData::CreateLambda([this, &UAFStateTreeContext](const FStateTreeExecutionContext& Context, const UStateTree* InStateTree, TArrayView<const FStateTreeExternalDataDesc> ExternalDataDescs, TArrayView<FStateTreeDataView> OutDataViews) -> bool
+			{
+				for (int32 Index = 0; Index < ExternalDataDescs.Num(); Index++)
+				{
+					const FStateTreeExternalDataDesc& ItemDesc = ExternalDataDescs[Index];
+					if (ItemDesc.Struct != nullptr)
+					{
+						if (ItemDesc.Struct->IsChildOf(FUAFStateTreeContext::StaticStruct()))
+						{
+							OutDataViews[Index] = FStateTreeDataView(FUAFStateTreeContext::StaticStruct(), reinterpret_cast<uint8*>(&UAFStateTreeContext));
+						}
+					}
+				}
+
+				return true;
+			}));
+
+		if (StateTreeExecutionContext.IsValid())
+		{
+			StateTreeExecutionContext.Start(); // todo: call Stop in destructor?
+		}
+	}
+
+}
+
+void FUAFStateTreeNode::PreUpdate(FUAFAnimGraphUpdateContext& Context)
+{
+	SCOPED_NAMED_EVENT(Controller_Update_StateTree, FColor::Blue);
+
+	if (HasChildren())
+	{
+		if (FUAFAnimNode* Child = GetFirstChild().GetReference())
+		{
+			if (IUAFTransitionNode* TransitionInstance = Child->GetInterface<IUAFTransitionNode>())
+			{
+				if (TransitionInstance->IsComplete())
+				{
+					// Our transition is complete, prune it
+					FUAFAnimNodePtr Target = TransitionInstance->ReleaseTarget();
+					SetChildAt(Target, 0);
+				}
+			}
+		}
+	}
+	
+	if (StateTree)
+	{
+		UObject* Owner = GetTransientPackage();
+#if ENABLE_ANIM_DEBUG 
+		Owner = const_cast<UObject*>(Context.GetHostObject());
+		// @TODO: Makes unique, but breaks visual logger
+		Owner = Owner ? Owner->GetOuter() : Owner;
+
+		if (!Owner)
+		{
+			return;
+		}
+#endif // ENABLE_ANIM_DEBUG 
+	
+		FStateTreeExecutionContext StateTreeExecutionContext(*Owner, *StateTree, InstanceData);
+	
+		if (StateTreeExecutionContext.GetLastTickStatus() != EStateTreeRunStatus::Failed)
+		{
+			FUAFStateTreeNodeContext UAFStateTreeContext(Context, this);
+			StateTreeExecutionContext.SetContextDataByName(UStateTreeAnimNextSchema::AnimStateTreeExecutionContextName, FStateTreeDataView(FUAFStateTreeContext::StaticStruct(), reinterpret_cast<uint8*>(&UAFStateTreeContext)));
+			StateTreeExecutionContext.SetExternalGlobalParameters(&StateTreeExternalParameters);
+#if OBJECT_TRACE_ENABLED
+			StateTreeExecutionContext.SetOuterTraceId(FObjectTrace::GetObjectId(Owner)); // todo setting outer id to outer object for now
+#endif
+			StateTreeExecutionContext.SetCollectExternalDataCallback(FOnCollectStateTreeExternalData::CreateLambda([this, &UAFStateTreeContext]( const FStateTreeExecutionContext& Context, const UStateTree* InStateTree, TArrayView<const FStateTreeExternalDataDesc> ExternalDataDescs, TArrayView<FStateTreeDataView> OutDataViews) -> bool
+			{
+				for (int32 Index = 0; Index < ExternalDataDescs.Num(); Index++)
+				{
+					const FStateTreeExternalDataDesc& ItemDesc = ExternalDataDescs[Index];
+					if (ItemDesc.Struct != nullptr)
+					{
+					   if (ItemDesc.Struct->IsChildOf(FUAFStateTreeContext::StaticStruct()))
+					   {
+						   OutDataViews[Index] = FStateTreeDataView(FUAFStateTreeContext::StaticStruct(), reinterpret_cast<uint8*>(&UAFStateTreeContext));
+					   }
+				   }
+			   }
+		
+			   return true;
+		   }));
+
+			StateTreeExecutionContext.Tick(Context.GetDeltaTime());
+		}
+	}
+}
+
+void* FUAFStateTreeNode::GetInterface(FUAFAnimNodeInterfaceId Id)
+{
+	if (Id == IUAFTransitionNode::InterfaceId)
+	{
+		// temp fix until we add interface "Search Depth" parameter:
+		// because this node may contain Transitions, and we call GetInterface to check for Transitions to prune...
+		// make sure the Transition Interface is never passed back from this node's child
+		return nullptr;
+	}
+
+	if (HasChildren())
+	{
+		if (FUAFAnimNodePtr Child = GetFirstChild())
+		{
+			return Child->GetInterface(Id);
+		}
+	}
+
+	return nullptr;
+}
+
+void FUAFStateTreeNode::AddReferencedObjects(FUAFAnimNode* This, FReferenceCollector& Collector)
+{
+	FUAFStateTreeNode* That = static_cast<FUAFStateTreeNode*>(This);
+	Collector.AddReferencedObject(That->StateTree);
+}
+
+void FUAFStateTreeNode::BlendTo(FUAFAnimGraphUpdateContext& Context, TConstStructView<FUAFGraphFactoryAsset> AssetData, FUAFTransitionNodeData* TransitionBlend)
+{
+	FUAFAnimNodePtr TargetInstance = FUAFGraphFactoryAssetAnimNodeFactory::CreateUAFAnimNodeFromObject(AssetData, Context);
+
+	if (HasChildren())
+	{
+		FUAFAnimNodePtr CurrentChild = GetFirstChild();
+		
+		// We already had a child, transition to our new target
+		// Remove our source since we'll re-parent it
+		SetChildAt(nullptr, 0);
+		
+		FUAFAnimNodePtr TransitionInstance;
+		if (TransitionBlend)
+		{
+			
+			TransitionInstance = TransitionBlend->CreateInstance(Context, CurrentChild, TargetInstance);
+		}
+		else
+		{
+			TransitionInstance = MakeAnimNode<FUAFSimpleTransition>(Context, CurrentChild, TargetInstance, 0.1f, EAlphaBlendOption::Linear);
+		}
+		
+		SetChildAt(TransitionInstance, 0);
+	}
+	else
+	{
+		AddChild(TargetInstance);
+	}
+}
+
+FUAFAnimNodePtr  FUAFStateTreeNodeData::CreateInstance(FUAFAnimGraphUpdateContext& Context) const
+{
+	if (StateTree && StateTree->StateTree)
+	{
+		return MakeAnimNode<FUAFStateTreeNode>(Context, StateTree->StateTree);
+	}
+	else
+	{
+		return nullptr;
+	}
+}
+
+void* FUAFStateTreeNodeData::GetInterface(FUAFAnimNodeInterfaceId Id)
+{
+	if (Id == IUAFAnimNodeDataHasAsset::InterfaceId)
+	{
+		return static_cast<IUAFAnimNodeDataHasAsset*>(this);
+	}
+
+	return nullptr;
+}
+	
+}
\ No newline at end of file
diff --git a/Engine/Plugins/Experimental/UAF/UAFStateTree/Source/UAFStateTree/Private/AnimNode/UAFStateTreeNodeConditions.cpp b/Engine/Plugins/Experimental/UAF/UAFStateTree/Source/UAFStateTree/Private/AnimNode/UAFStateTreeNodeConditions.cpp
new file mode 100644
--- /dev/null
+++ b/Engine/Plugins/Experimental/UAF/UAFStateTree/Source/UAFStateTree/Private/AnimNode/UAFStateTreeNodeConditions.cpp
@@ -0,0 +1,53 @@
+// Copyright Epic Games, Inc. All Rights Reserved.
+
+#include "AnimNode/UAFStateTreeNodeConditions.h"
+
+#include "AnimNextStateTreeSchema.h"
+#include "StateTreeExecutionContext.h"
+#include "UAFStateTreeNodeContext.h"
+#include "Script/UAFRigVMComponent.h"
+#include "Logging/LogScopedVerbosityOverride.h"
+
+#if ENABLE_ANIM_DEBUG 
+#include "Debugger/StateTreeRuntimeValidation.h"
+#endif // ENABLE_ANIM_DEBUG
+
+namespace UE::UAF::StateTree
+{
+	
+bool FUAFFloatCompareCondition::TestCondition(FStateTreeExecutionContext& Context) const
+{
+	FUAFStateTreeContext& ExecContext = Context.GetExternalData(ContextHandle);
+
+	if (FUAFAssetInstance* VariablesOwner = ExecContext.GetVariablesOwner())
+	{
+		double LeftValue;
+		if (VariablesOwner->GetVariable(Left, LeftValue) == EPropertyBagResult::Success)
+		{	
+			if (Compare == EUAFFloatCompareConditionType::Less)
+			{
+				return LeftValue < Right;
+			}
+			return LeftValue > Right;
+		}
+	}
+	return false;
+}
+
+bool FUAFEnumCompareCondition::TestCondition(FStateTreeExecutionContext& Context) const
+{
+	FUAFStateTreeContext& ExecContext = Context.GetExternalData(ContextHandle);
+
+	if (FUAFAssetInstance* VariablesOwner = ExecContext.GetVariablesOwner())
+	{
+		uint8 LeftValue;
+		if (VariablesOwner->GetVariable(Left, LeftValue) == EPropertyBagResult::Success)
+		{
+			return LeftValue == Right;
+		}
+	}
+
+	return false;
+}
+
+}
\ No newline at end of file
diff --git a/Engine/Plugins/Experimental/UAF/UAFStateTree/Source/UAFStateTree/Private/AnimNode/UAFStateTreeNodeContext.cpp b/Engine/Plugins/Experimental/UAF/UAFStateTree/Source/UAFStateTree/Private/AnimNode/UAFStateTreeNodeContext.cpp
new file mode 100644
--- /dev/null
+++ b/Engine/Plugins/Experimental/UAF/UAFStateTree/Source/UAFStateTree/Private/AnimNode/UAFStateTreeNodeContext.cpp
@@ -0,0 +1,53 @@
+// Copyright Epic Games, Inc. All Rights Reserved.
+
+#include "UAFStateTreeNodeContext.h"
+#include "AnimNode/UAFStateTreeNode.h"
+#include "UAF/AnimNodes/IUAFAnimNodeTimeline.h"
+
+namespace UE::UAF::StateTree
+{
+
+void FUAFStateTreeNodeContext::QueryPlaybackInfo(FPlaybackInfo& OutPlaybackInfo) const
+{
+	if (AnimNode)
+	{
+		// todo: make templated GetInterface calls work correctly when called on a specific node type.
+		if (const UE::UAF::IUAFAnimNodeTimeline* Timeline = AnimNode->FUAFAnimNode::GetInterface<IUAFAnimNodeTimeline>())
+		{
+			float Duration = Timeline->GetLength();
+			float CurrentTime = Timeline->GetCurrentTime();
+			OutPlaybackInfo.PlaybackRatio = Duration > 0.f ? CurrentTime / Duration : 1.0f;
+			OutPlaybackInfo.Duration = Duration;
+			// todo: handle negative playback rates (need something exposed from ITimeline)
+			OutPlaybackInfo.TimeLeft = Duration - CurrentTime;
+			OutPlaybackInfo.bIsLooping = Timeline->IsLooping();
+		}
+	}
+}
+
+bool FUAFStateTreeNodeContext::PushAssetOntoBlendStack(FGraphAssetHandleConstView InAsset, const FAlphaBlendArgs& InBlendArguments, const UUAFBlendProfile* BlendProfile) const
+{
+	using namespace UE::UAF;
+
+	if (Context && AnimNode)
+	{
+		if (InAsset.IsValid())
+		{
+			AnimNode->BlendTo(*Context, InAsset);
+			// todo: use InBlendArguments to create an appropriate Transition
+			return true;
+		}
+	}
+	return false;
+}
+
+FUAFAssetInstance* FUAFStateTreeNodeContext::GetVariablesOwner() const
+{
+	if (Context)
+	{
+		return Context->GetVariablesOwner();
+	}
+	return nullptr;
+}
+
+}
\ No newline at end of file
diff --git a/Engine/Plugins/Experimental/UAF/UAFStateTree/Source/UAFStateTree/Private/AnimNode/UAFStateTreeNodeContext.h b/Engine/Plugins/Experimental/UAF/UAFStateTree/Source/UAFStateTree/Private/AnimNode/UAFStateTreeNodeContext.h
new file mode 100644
--- /dev/null
+++ b/Engine/Plugins/Experimental/UAF/UAFStateTree/Source/UAFStateTree/Private/AnimNode/UAFStateTreeNodeContext.h
@@ -0,0 +1,54 @@
+// Copyright Epic Games, Inc. All Rights Reserved.
+
+#pragma once
+
+#include "UAFStateTreeContext.h"
+#include "UAFStateTreeNodeContext.generated.h"
+
+namespace UE::UAF
+{
+class FUAFAnimGraphUpdateContext;
+}
+
+class UUAFBlendProfile;
+
+namespace UE::UAF::StateTree
+{
+class FUAFStateTreeNode;
+
+USTRUCT()
+struct FUAFStateTreeNodeContext : public FUAFStateTreeContext
+{	
+	GENERATED_BODY()
+	
+	FUAFStateTreeNodeContext() 
+		: Context(nullptr), AnimNode(nullptr)
+	{
+	}
+
+	FUAFStateTreeNodeContext(FUAFAnimGraphUpdateContext& InContext, FUAFStateTreeNode* InAnimNode)
+		: Context(&InContext), AnimNode(InAnimNode)
+	{
+	}
+
+	virtual bool PushAssetOntoBlendStack(FGraphAssetHandleConstView InAsset, const FAlphaBlendArgs& InBlendArguments, const UUAFBlendProfile* BlendProfile) const override;
+    virtual void QueryPlaybackInfo(FPlaybackInfo& OutPlaybackInfo) const override;
+	virtual FUAFAssetInstance* GetVariablesOwner() const override;
+
+	UE::UAF::FUAFAnimGraphUpdateContext* GetContext() const
+	{
+		return Context;
+	}
+
+	FUAFStateTreeNode* GetAnimNode() const
+	{
+		return AnimNode;
+	}
+
+private:
+
+	FUAFAnimGraphUpdateContext* Context;
+	FUAFStateTreeNode* AnimNode;
+};
+
+}
diff --git a/Engine/Plugins/Experimental/UAF/UAFStateTree/Source/UAFStateTree/Private/AnimStateTreeTrait.cpp b/Engine/Plugins/Experimental/UAF/UAFStateTree/Source/UAFStateTree/Private/AnimStateTreeTrait.cpp
--- a/Engine/Plugins/Experimental/UAF/UAFStateTree/Source/UAFStateTree/Private/AnimStateTreeTrait.cpp
+++ b/Engine/Plugins/Experimental/UAF/UAFStateTree/Source/UAFStateTree/Private/AnimStateTreeTrait.cpp
@@ -2,7 +2,7 @@
 
 #include "AnimStateTreeTrait.h"
 
-#include "AnimNextStateTreeContext.h"
+#include "UAFStateTreeTraitContext.h"
 #include "AnimNextStateTreeSchema.h"
 #include "StateTreeExecutionContext.h"
 #include "Script/UAFRigVMComponent.h"
@@ -43,8 +43,8 @@ namespace UE::UAF
 			FAnimNextGraphInstance* OwnerGraphInstance = &Binding.GetTraitPtr().GetNodeInstance()->GetOwner();
 			if(OwnerGraphInstance && StateTreeExecutionContext.GetLastTickStatus() != EStateTreeRunStatus::Failed)
 			{
-				FAnimNextStateTreeTraitContext TraitContext(*const_cast<FExecutionContext*>(&Context), &Binding);
-				StateTreeExecutionContext.SetContextDataByName(UStateTreeAnimNextSchema::AnimStateTreeExecutionContextName, FStateTreeDataView(FAnimNextStateTreeTraitContext::StaticStruct(), reinterpret_cast<uint8*>(&TraitContext)));
+				FUAFStateTreeTraitContext TraitContext(*const_cast<FExecutionContext*>(&Context), &Binding);
+				StateTreeExecutionContext.SetContextDataByName(UStateTreeAnimNextSchema::AnimStateTreeExecutionContextName, FStateTreeDataView(FUAFStateTreeTraitContext::StaticStruct(), reinterpret_cast<uint8*>(&TraitContext)));
 				StateTreeExecutionContext.SetExternalGlobalParameters(&StateTreeExternalParameters);
 
 				StateTreeExecutionContext.SetCollectExternalDataCallback(FOnCollectStateTreeExternalData::CreateLambda([this, &TraitContext]( const FStateTreeExecutionContext& Context, const UStateTree* InStateTree, TArrayView<const FStateTreeExternalDataDesc> ExternalDataDescs, TArrayView<FStateTreeDataView> OutDataViews) -> bool
@@ -54,9 +54,9 @@ namespace UE::UAF
 						const FStateTreeExternalDataDesc& ItemDesc = ExternalDataDescs[Index];
 						if (ItemDesc.Struct != nullptr)
 						{
-							if (ItemDesc.Struct->IsChildOf(FAnimNextStateTreeTraitContext::StaticStruct()))
+							if (ItemDesc.Struct->IsChildOf(FUAFStateTreeTraitContext::StaticStruct()))
 							{
-								OutDataViews[Index] = FStateTreeDataView(FAnimNextStateTreeTraitContext::StaticStruct(), reinterpret_cast<uint8*>(&TraitContext));
+								OutDataViews[Index] = FStateTreeDataView(FUAFStateTreeTraitContext::StaticStruct(), reinterpret_cast<uint8*>(&TraitContext));
 							}
 						}
 					}
@@ -93,8 +93,8 @@ namespace UE::UAF
 
 			FStateTreeExecutionContext StateTreeExecutionContext(*Owner, *InstanceData->StateTree, InstanceData->InstanceData);
 
-			FAnimNextStateTreeTraitContext TraitContext(Context, Binding);
-			StateTreeExecutionContext.SetContextDataByName(UStateTreeAnimNextSchema::AnimStateTreeExecutionContextName, FStateTreeDataView(FAnimNextStateTreeTraitContext::StaticStruct(), &TraitContext));
+			FUAFStateTreeTraitContext TraitContext(Context, Binding);
+			StateTreeExecutionContext.SetContextDataByName(UStateTreeAnimNextSchema::AnimStateTreeExecutionContextName, FStateTreeDataView(FUAFStateTreeContext::StaticStruct(), &TraitContext));
 
 			FAnimNextGraphInstance& OwnerGraphInstance = Binding.GetTraitPtr().GetNodeInstance()->GetOwner();
 			StateTreeExecutionContext.SetOuterTraceId(OwnerGraphInstance.GetUniqueId());
@@ -178,9 +178,9 @@ namespace UE::UAF
 					const FStateTreeExternalDataDesc& ItemDesc = ExternalDataDescs[Index];
 					if (ItemDesc.Struct != nullptr)
 					{
-						if (ItemDesc.Struct->IsChildOf(FAnimNextStateTreeTraitContext::StaticStruct()))
+						if (ItemDesc.Struct->IsChildOf(FUAFStateTreeContext::StaticStruct()))
 						{
-							OutDataViews[Index] = FStateTreeDataView(FAnimNextStateTreeTraitContext::StaticStruct(), reinterpret_cast<uint8*>(&TraitContext));
+							OutDataViews[Index] = FStateTreeDataView(FUAFStateTreeContext::StaticStruct(), reinterpret_cast<uint8*>(&TraitContext));
 						}
 					}
 				}
@@ -216,8 +216,8 @@ namespace UE::UAF
 			FAnimNextGraphInstance& OwnerGraphInstance = Binding.GetTraitPtr().GetNodeInstance()->GetOwner();
 			if (StateTreeExecutionContext.GetLastTickStatus() != EStateTreeRunStatus::Failed)
 			{
-				FAnimNextStateTreeTraitContext TraitContext(Context, Binding);
-				StateTreeExecutionContext.SetContextDataByName(UStateTreeAnimNextSchema::AnimStateTreeExecutionContextName, FStateTreeDataView(FAnimNextStateTreeTraitContext::StaticStruct(), reinterpret_cast<uint8*>(&TraitContext)));
+				FUAFStateTreeTraitContext TraitContext(Context, Binding);
+				StateTreeExecutionContext.SetContextDataByName(UStateTreeAnimNextSchema::AnimStateTreeExecutionContextName, FStateTreeDataView(FUAFStateTreeContext::StaticStruct(), reinterpret_cast<uint8*>(&TraitContext)));
 				StateTreeExecutionContext.SetExternalGlobalParameters(&InstanceData->StateTreeExternalParameters);
 				StateTreeExecutionContext.SetOuterTraceId(OwnerGraphInstance.GetUniqueId());
 				StateTreeExecutionContext.SetCollectExternalDataCallback(FOnCollectStateTreeExternalData::CreateLambda([this, &TraitContext]( const FStateTreeExecutionContext& Context, const UStateTree* StateTree, TArrayView<const FStateTreeExternalDataDesc> ExternalDataDescs, TArrayView<FStateTreeDataView> OutDataViews) -> bool
@@ -227,9 +227,9 @@ namespace UE::UAF
 						const FStateTreeExternalDataDesc& ItemDesc = ExternalDataDescs[Index];
 						if (ItemDesc.Struct != nullptr)
 						{
-							if (ItemDesc.Struct->IsChildOf(FAnimNextStateTreeTraitContext::StaticStruct()))
+							if (ItemDesc.Struct->IsChildOf(FUAFStateTreeContext::StaticStruct()))
 							{
-								OutDataViews[Index] = FStateTreeDataView(FAnimNextStateTreeTraitContext::StaticStruct(), reinterpret_cast<uint8*>(&TraitContext));
+								OutDataViews[Index] = FStateTreeDataView(FUAFStateTreeContext::StaticStruct(), reinterpret_cast<uint8*>(&TraitContext));
 							}
 						}
 					}
diff --git a/Engine/Plugins/Experimental/UAF/UAFStateTree/Source/UAFStateTree/Private/Conditions/AnimNextStateTreeRigVMConditionBase.cpp b/Engine/Plugins/Experimental/UAF/UAFStateTree/Source/UAFStateTree/Private/Conditions/AnimNextStateTreeRigVMConditionBase.cpp
--- a/Engine/Plugins/Experimental/UAF/UAFStateTree/Source/UAFStateTree/Private/Conditions/AnimNextStateTreeRigVMConditionBase.cpp
+++ b/Engine/Plugins/Experimental/UAF/UAFStateTree/Source/UAFStateTree/Private/Conditions/AnimNextStateTreeRigVMConditionBase.cpp
@@ -8,7 +8,7 @@
 #include "StateTreeExecutionContext.h"
 #include "StructUtils/PropertyBag.h"
 #include "TraitCore/ExecutionContext.h"
-#include "AnimNextStateTreeContext.h"
+#include "UAFStateTreeContext.h"
 #include "Script/UAFRigVMComponent.h"
 
 #if WITH_EDITOR
@@ -39,11 +39,15 @@ bool FAnimNextStateTreeRigVMConditionBase::TestCondition(FStateTreeExecutionCont
 {
 	if (FAnimNextStateTreeRigVMConditionInstanceData* InstanceData = Context.GetInstanceDataPtr<FAnimNextStateTreeRigVMConditionInstanceData>(*this))
 	{
-		FAnimNextStateTreeTraitContext& ExecContext = Context.GetExternalData(TraitContextHandle);
-		FAnimNextGraphInstance& GraphInstance = ExecContext.GetGraphInstance();
+		FUAFStateTreeContext& ExecContext = Context.GetExternalData(TraitContextHandle);
+		FUAFAssetInstance* GraphInstance = ExecContext.GetVariablesOwner();
+		if (!GraphInstance)
+		{
+			return false;
+		}
 
 		// Need const cast as VM execution isn't const.
-		if (UUAFRigVMAsset* RigVMInstance = const_cast<UUAFRigVMAsset*>(GraphInstance.GetAsset<UUAFRigVMAsset>()))
+		if (UUAFRigVMAsset* RigVMInstance = const_cast<UUAFRigVMAsset*>(GraphInstance->GetAsset<UUAFRigVMAsset>()))
 		{
 #if WITH_EDITOR
 			if (!InstanceData->FunctionHandle.IsValidForVM(RigVMInstance->GetRigVM()->GetVMHash()))
@@ -59,11 +63,11 @@ bool FAnimNextStateTreeRigVMConditionBase::TestCondition(FStateTreeExecutionCont
 				return false;
 			}
 
-			FUAFRigVMComponent& RigVMComponent = GraphInstance.GetComponent<FUAFRigVMComponent>();
+			FUAFRigVMComponent& RigVMComponent = GraphInstance->GetComponent<FUAFRigVMComponent>();
 			FRigVMExtendedExecuteContext& ExtendedExecuteContext = RigVMComponent.GetExtendedExecuteContext();
 			FAnimNextExecuteContext& AnimNextContext = ExtendedExecuteContext.GetPublicDataSafe<FAnimNextExecuteContext>();
 
-			FAnimNextGraphContextData ContextData(GraphInstance, CallFunctionInfo.FunctionEvent, AnimNextContext.GetDeltaTime<float>());
+			FUAFAssetContextData ContextData(*GraphInstance, CallFunctionInfo.FunctionEvent, AnimNextContext.GetDeltaTime<float>());
 			UE::UAF::FScopedExecuteContextData ContextDataScope(AnimNextContext, ContextData);
 
 			RigVMInstance->CallFunctionHandle(InstanceData->FunctionHandle, ExtendedExecuteContext, InstanceData->Parameters);
diff --git a/Engine/Plugins/Experimental/UAF/UAFStateTree/Source/UAFStateTree/Private/Tasks/AnimNextStateTreeGraphInstanceTask.cpp b/Engine/Plugins/Experimental/UAF/UAFStateTree/Source/UAFStateTree/Private/Tasks/AnimNextStateTreeGraphInstanceTask.cpp
--- a/Engine/Plugins/Experimental/UAF/UAFStateTree/Source/UAFStateTree/Private/Tasks/AnimNextStateTreeGraphInstanceTask.cpp
+++ b/Engine/Plugins/Experimental/UAF/UAFStateTree/Source/UAFStateTree/Private/Tasks/AnimNextStateTreeGraphInstanceTask.cpp
@@ -4,13 +4,15 @@
 
 #include "StateTreeLinker.h"
 #include "StateTreeExecutionContext.h"
-#include "AnimNextStateTreeContext.h"
+#include "UAFStateTreeContext.h"
 #include "Asset/UAFAnimGraphAssetData.h"
 #include "Factory/AnimGraphFactory.h"
 #include "Traits/BlendSpacePlayerTraitData.h"
 #include "Traits/SequencePlayerTraitData.h"
 #include "UAF/UAFAssetFactory.h"
 #include "UObject/FortniteMainBranchObjectVersion.h"
+#include "UAF/AnimNodeCore/UAFAnimNode.h"
+#include "UAF/AnimNodes/IUAFAnimNodeTimeline.h"
 
 #define LOCTEXT_NAMESPACE "AnimNextStateTreeGraphInstanceTask"
 
@@ -172,7 +174,7 @@ bool FAnimNextStateTreeGraphInstanceTask::Link(FStateTreeLinker& Linker)
 
 EStateTreeRunStatus FAnimNextStateTreeGraphInstanceTask::EnterState(FStateTreeExecutionContext& Context, const FStateTreeTransitionResult& Transition) const
 {
-	FAnimNextStateTreeTraitContext& ExecContext = Context.GetExternalData(TraitContextHandle);
+	FUAFStateTreeContext& ExecContext = Context.GetExternalData(TraitContextHandle);
 	FInstanceDataType& InstanceData = Context.GetInstanceData(*this);
 
 	if (ExecContext.PushAssetOntoBlendStack(InstanceData.AssetData, InstanceData.BlendOptions, InstanceData.BlendProfile))
@@ -186,14 +188,18 @@ EStateTreeRunStatus FAnimNextStateTreeGraphInstanceTask::EnterState(FStateTreeEx
 EStateTreeRunStatus FAnimNextStateTreeGraphInstanceTask::Tick(FStateTreeExecutionContext& Context, const float DeltaTime) const
 {
 	FInstanceDataType& InstanceData = Context.GetInstanceData(*this);
+
+	FUAFStateTreeContext& ExecContext = Context.GetExternalData(TraitContextHandle);
+
 	if (InstanceData.bContinueTicking)
 	{
-		FAnimNextStateTreeTraitContext& ExecContext = Context.GetExternalData(TraitContextHandle);
-
-		InstanceData.PlaybackRatio = ExecContext.QueryPlaybackRatio();
-		InstanceData.Duration = ExecContext.QueryDuration();
-		InstanceData.TimeLeft = ExecContext.QueryTimeLeft();
-		InstanceData.bIsLooping = ExecContext.QueryIsLooping();
+		FUAFStateTreeContext::FPlaybackInfo PlaybackInfo;
+		ExecContext.QueryPlaybackInfo(PlaybackInfo);
+		
+		InstanceData.PlaybackRatio = PlaybackInfo.PlaybackRatio;
+		InstanceData.Duration = PlaybackInfo.Duration;
+		InstanceData.TimeLeft = PlaybackInfo.TimeLeft;
+		InstanceData.bIsLooping = PlaybackInfo.bIsLooping;
 
 		// We can get a fallback timeline with duration of 0.0f during graph init / compile.
 		// Avoid finishing the task for this edge case
diff --git a/Engine/Plugins/Experimental/UAF/UAFStateTree/Source/UAFStateTree/Private/Tasks/AnimNextStateTreeRigVMTaskBase.cpp b/Engine/Plugins/Experimental/UAF/UAFStateTree/Source/UAFStateTree/Private/Tasks/AnimNextStateTreeRigVMTaskBase.cpp
--- a/Engine/Plugins/Experimental/UAF/UAFStateTree/Source/UAFStateTree/Private/Tasks/AnimNextStateTreeRigVMTaskBase.cpp
+++ b/Engine/Plugins/Experimental/UAF/UAFStateTree/Source/UAFStateTree/Private/Tasks/AnimNextStateTreeRigVMTaskBase.cpp
@@ -8,7 +8,7 @@
 #include "StateTreeExecutionContext.h"
 #include "StructUtils/PropertyBag.h"
 #include "TraitCore/ExecutionContext.h"
-#include "AnimNextStateTreeContext.h"
+#include "UAFStateTreeContext.h"
 #include "AnimNextRigVMAsset.h"
 #include "Script/UAFRigVMComponent.h"
 
@@ -69,11 +69,15 @@ EStateTreeRunStatus FAnimNextStateTreeRigVMTaskBase::CallFunction(FStateTreeExec
 {
 	if (FAnimNextStateTreeRigVMTaskInstanceData* InstanceData = Context.GetInstanceDataPtr<FAnimNextStateTreeRigVMTaskInstanceData>(*this))
 	{
-		FAnimNextStateTreeTraitContext& ExecContext = Context.GetExternalData(TraitContextHandle);
-		FAnimNextGraphInstance& GraphInstance = ExecContext.GetGraphInstance();
+		FUAFStateTreeContext& ExecContext = Context.GetExternalData(TraitContextHandle);
+		FUAFAssetInstance* GraphInstance = ExecContext.GetVariablesOwner();
+		if (!GraphInstance)
+		{
+			return EStateTreeRunStatus::Failed;
+		}
 
 		// Need const cast as VM execution isn't const.
-		if (UUAFRigVMAsset* RigVMAsset = const_cast<UUAFRigVMAsset*>(GraphInstance.GetAsset<UUAFRigVMAsset>()))
+		if (UUAFRigVMAsset* RigVMAsset = const_cast<UUAFRigVMAsset*>(GraphInstance->GetAsset<UUAFRigVMAsset>()))
 		{
 #if WITH_EDITOR
 			if (!InstanceData->FunctionHandle.IsValidForVM(RigVMAsset->GetRigVM()->GetVMHash()))
@@ -90,11 +94,11 @@ EStateTreeRunStatus FAnimNextStateTreeRigVMTaskBase::CallFunction(FStateTreeExec
 				return EStateTreeRunStatus::Failed;
 			}
 
-			FUAFRigVMComponent& RigVMComponent = GraphInstance.GetComponent<FUAFRigVMComponent>();
+			FUAFRigVMComponent& RigVMComponent = GraphInstance->GetComponent<FUAFRigVMComponent>();
 			FRigVMExtendedExecuteContext& ExtendedExecuteContext = RigVMComponent.GetExtendedExecuteContext();
 			FAnimNextExecuteContext& AnimNextContext = ExtendedExecuteContext.GetPublicDataSafe<FAnimNextExecuteContext>();
 
-			FAnimNextGraphContextData ContextData(GraphInstance, CallFunctionInfo.FunctionEvent, AnimNextContext.GetDeltaTime<float>());
+			FUAFAssetContextData ContextData(*GraphInstance, CallFunctionInfo.FunctionEvent, AnimNextContext.GetDeltaTime<float>());
 			UE::UAF::FScopedExecuteContextData ContextDataScope(AnimNextContext, ContextData);
 
 			RigVMAsset->CallFunctionHandle(InstanceData->FunctionHandle, ExtendedExecuteContext, InstanceData->Parameters);
diff --git a/Engine/Plugins/Experimental/UAF/UAFStateTree/Source/UAFStateTree/Private/UAFStateTreeContext.h b/Engine/Plugins/Experimental/UAF/UAFStateTree/Source/UAFStateTree/Private/UAFStateTreeContext.h
new file mode 100644
--- /dev/null
+++ b/Engine/Plugins/Experimental/UAF/UAFStateTree/Source/UAFStateTree/Private/UAFStateTreeContext.h
@@ -0,0 +1,32 @@
+// Copyright Epic Games, Inc. All Rights Reserved.
+
+#pragma once
+
+#include "UAF/UAFAssetData.h"
+
+#include "UAFStateTreeContext.generated.h"
+
+struct FAlphaBlendArgs;
+struct FUAFAssetInstance;
+class UUAFBlendProfile;
+
+USTRUCT()
+struct FUAFStateTreeContext
+{
+	GENERATED_BODY()
+	virtual ~FUAFStateTreeContext() {};
+	
+	struct FPlaybackInfo
+	{
+		float PlaybackRatio = 0.f;
+    	float TimeLeft = 0.f;
+    	float Duration = 0.f;
+    	bool bIsLooping = false;
+	};
+
+	virtual bool PushAssetOntoBlendStack(UE::UAF::FGraphAssetHandleConstView InAsset, const FAlphaBlendArgs& InBlendArguments, const UUAFBlendProfile* InBlendProfile) const { return false; };
+    virtual void QueryPlaybackInfo(FPlaybackInfo& OutPlaybackInfo) const {};
+	virtual FUAFAssetInstance* GetVariablesOwner() const { return nullptr; }
+};
+
+
diff --git a/Engine/Plugins/Experimental/UAF/UAFStateTree/Source/UAFStateTree/Private/UAFStateTreeRigVMUtils.cpp b/Engine/Plugins/Experimental/UAF/UAFStateTree/Source/UAFStateTree/Private/UAFStateTreeRigVMUtils.cpp
--- a/Engine/Plugins/Experimental/UAF/UAFStateTree/Source/UAFStateTree/Private/UAFStateTreeRigVMUtils.cpp
+++ b/Engine/Plugins/Experimental/UAF/UAFStateTree/Source/UAFStateTree/Private/UAFStateTreeRigVMUtils.cpp
@@ -5,7 +5,6 @@
 #if WITH_EDITOR
 
 #include "AnimNextFunctionHandle.h"
-#include "AnimNextStateTreeContext.h"
 #include "AnimNextStateTreeEditorOnlyTypes.h"
 #include "AnimNextRigVMAsset.h"
 #include "Graph/AnimNextGraphContextData.h"
diff --git a/Engine/Plugins/Experimental/UAF/UAFStateTree/Source/UAFStateTree/Private/UAFStateTreeTraitContext.cpp b/Engine/Plugins/Experimental/UAF/UAFStateTree/Source/UAFStateTree/Private/UAFStateTreeTraitContext.cpp
new file mode 100644
--- /dev/null
+++ b/Engine/Plugins/Experimental/UAF/UAFStateTree/Source/UAFStateTree/Private/UAFStateTreeTraitContext.cpp
@@ -0,0 +1,62 @@
+// Copyright Epic Games, Inc. All Rights Reserved.
+
+#include "UAFStateTreeTraitContext.h"
+#include "Factory/AnimGraphFactory.h"
+#include "TraitCore/NodeInstance.h"
+#include "TraitCore/TraitStackBinding.h"
+#include "TraitInterfaces/IBlendStack.h"
+#include "TraitInterfaces/ITimeline.h"
+
+bool FUAFStateTreeTraitContext::PushAssetOntoBlendStack(UE::UAF::FGraphAssetHandleConstView InAsset, const FAlphaBlendArgs& InBlendArguments, const UUAFBlendProfile* BlendProfile) const
+{
+	using namespace UE::UAF;
+
+	TTraitBinding<IBlendStack> BlendStackBinding;
+	if(!Binding->GetStackInterface<IBlendStack>(BlendStackBinding))
+	{
+		return false;
+	}
+
+	// Instantiate/get a graph from the params
+	FAnimNextFactoryParams Params = FAnimGraphFactory::GetDefaultParamsForAsset(InAsset);
+	const UUAFAnimGraph* AnimationGraph = FAnimGraphFactory::GetOrBuildGraph(InAsset);
+	if(AnimationGraph == nullptr)
+	{
+		return false;
+	}
+
+	IBlendStack::FGraphRequest NewGraphRequest;
+	NewGraphRequest.BlendArgs = InBlendArguments;
+	// Todo: change to asset handle?
+	TArray<const UObject*> Objects;
+	InAsset.GetPtr()->GetObjectReferences(Objects);
+	NewGraphRequest.FactoryObject = Objects.Num() > 0 ? Objects[0] : nullptr;
+	NewGraphRequest.AnimationGraph = AnimationGraph;
+	NewGraphRequest.FactoryParams = Params;
+	NewGraphRequest.BlendProfile = BlendProfile;
+
+	BlendStackBinding.PushGraph(*Context, MoveTemp(NewGraphRequest));
+
+	return true;
+}
+
+void FUAFStateTreeTraitContext::QueryPlaybackInfo(FPlaybackInfo& OutPlaybackInfo) const
+{
+	UE::UAF::TTraitBinding<UE::UAF::ITimeline> Timeline;
+	if (Binding->GetStackInterface<UE::UAF::ITimeline>(Timeline))
+	{
+		UE::UAF::FTimelineState TimelineState;
+		if (Timeline.GetState(*Context, TimelineState))
+		{
+			OutPlaybackInfo.PlaybackRatio = TimelineState.GetPositionRatio();
+			OutPlaybackInfo.TimeLeft = TimelineState.GetTimeLeft();
+			OutPlaybackInfo.Duration = TimelineState.GetDuration();
+			OutPlaybackInfo.bIsLooping = TimelineState.IsLooping();
+		}
+	}
+}
+
+FUAFAssetInstance* FUAFStateTreeTraitContext::GetVariablesOwner() const
+{
+	return &Binding->GetTraitPtr().GetNodeInstance()->GetOwner();
+}
\ No newline at end of file
diff --git a/Engine/Plugins/Experimental/UAF/UAFStateTree/Source/UAFStateTree/Private/UAFStateTreeTraitContext.h b/Engine/Plugins/Experimental/UAF/UAFStateTree/Source/UAFStateTree/Private/UAFStateTreeTraitContext.h
new file mode 100644
--- /dev/null
+++ b/Engine/Plugins/Experimental/UAF/UAFStateTree/Source/UAFStateTree/Private/UAFStateTreeTraitContext.h
@@ -0,0 +1,55 @@
+// Copyright Epic Games, Inc. All Rights Reserved.
+
+#pragma once
+
+#include "UAFStateTreeContext.h"
+#include "AnimStateTreeTrait.h"
+#include "UAFStateTreeTraitContext.generated.h"
+
+namespace UE::UAF
+{
+	struct FTraitBinding;
+	struct FExecutionContext;
+	struct FStateTreeTrait;
+}
+
+class UUAFAnimGraph;
+struct FUAFAssetInstance;
+
+USTRUCT()
+struct FUAFStateTreeTraitContext : public FUAFStateTreeContext
+{	
+	GENERATED_BODY()
+	
+	friend UE::UAF::FStateTreeTrait;
+	friend UE::UAF::FStateTreeTrait::FInstanceData;
+
+	FUAFStateTreeTraitContext() = default;
+
+	FUAFStateTreeTraitContext(UE::UAF::FExecutionContext& InContext, const UE::UAF::FTraitBinding& InBinding)
+		: Context(&InContext)
+		, Binding(&InBinding)
+	{}
+	
+	virtual bool PushAssetOntoBlendStack(UE::UAF::FGraphAssetHandleConstView InAsset, const FAlphaBlendArgs& InBlendArguments, const UUAFBlendProfile* BlendProfile) const override;
+	
+    virtual void QueryPlaybackInfo(FPlaybackInfo& OutPlaybackInfo) const override;
+	virtual FUAFAssetInstance* GetVariablesOwner() const override;
+
+	UE::UAF::FExecutionContext& GetContext() const
+	{
+		check(Context != nullptr);
+		return *Context;
+	}
+
+	const UE::UAF::FTraitBinding& GetBinding() const
+	{
+		check(Binding != nullptr);
+		return *Binding;
+	}
+
+private:
+
+	UE::UAF::FExecutionContext* Context = nullptr;
+	const UE::UAF::FTraitBinding* Binding = nullptr;
+};
\ No newline at end of file
diff --git a/Engine/Plugins/Experimental/UAF/UAFStateTree/Source/UAFStateTree/UAFStateTree.Build.cs b/Engine/Plugins/Experimental/UAF/UAFStateTree/Source/UAFStateTree/UAFStateTree.Build.cs
--- a/Engine/Plugins/Experimental/UAF/UAFStateTree/Source/UAFStateTree/UAFStateTree.Build.cs
+++ b/Engine/Plugins/Experimental/UAF/UAFStateTree/Source/UAFStateTree/UAFStateTree.Build.cs
@@ -6,10 +6,13 @@ namespace UnrealBuildTool.Rules
 	{
 		public UAFStateTree(ReadOnlyTargetRules Target) : base(Target)
 		{
+			bAllowUETypesInNamespaces = true;
+			
 			PrivateDependencyModuleNames.AddRange(
 				new string[]
 				{
 					"UAFAnimGraph",
+					"UAFAnimNode",
 					"RigVM"
 				});
 
diff --git a/Engine/Plugins/Experimental/UAF/UAFStateTree/Source/UAFStateTreeEditor/Private/AnimNextStateTreeEditorModule.cpp b/Engine/Plugins/Experimental/UAF/UAFStateTree/Source/UAFStateTreeEditor/Private/AnimNextStateTreeEditorModule.cpp
--- a/Engine/Plugins/Experimental/UAF/UAFStateTree/Source/UAFStateTreeEditor/Private/AnimNextStateTreeEditorModule.cpp
+++ b/Engine/Plugins/Experimental/UAF/UAFStateTree/Source/UAFStateTreeEditor/Private/AnimNextStateTreeEditorModule.cpp
@@ -28,6 +28,8 @@
 #include "StateTreeAssetCompilationHandler.h"
 #include "StateTreeEditorModule.h"
 #include "UAFStateTreeEditorSchema.h"
+#include "UAFAnimNodeEditor/Public/UAFAnimNodeDataFactory.h"
+#include "AnimNode/UAFStateTreeNode.h"
 
 #define LOCTEXT_NAMESPACE "AnimNextStateTreeEditorModule"
 
@@ -344,6 +346,14 @@ void FAnimNextStateTreeEditorModule::StartupModule()
 	FStateTreeEditorModule& StateTreeEditorModule = FStateTreeEditorModule::GetModule();
 	StateTreeEditorModule.RegisterEditorDataClass(UStateTreeAnimNextSchema::StaticClass(), UAnimNextStateTreeTreeEditorData::StaticClass());
 	StateTreeEditorModule.RegisterEditorSchemaClass(UStateTreeAnimNextSchema::StaticClass(), UUAFStateTreeEditorSchema::StaticClass());
+
+	StateTreeClassPath = FUAFAnimNodeDataFactory::RegisterAsset<FUAFStateTreeNodeData, UAnimNextStateTree>(
+	[](UAnimNextStateTree* StateTree)
+		{
+			FUAFStateTreeNodeData StateTreeNode;
+			StateTreeNode.SetStateTree(StateTree);
+			return StateTreeNode;
+		});
 }
 
 void FAnimNextStateTreeEditorModule::ShutdownModule()
@@ -369,6 +379,8 @@ void FAnimNextStateTreeEditorModule::ShutdownModule()
 			StateTreeEditorModule->UnregisterEditorSchemaClass(UStateTreeAnimNextSchema::StaticClass());
 		}
     }
+    
+    FUAFAnimNodeDataFactory::UnregisterAsset(StateTreeClassPath);
 }
 
 IMPLEMENT_MODULE(FAnimNextStateTreeEditorModule, UAFStateTreeEditor)
diff --git a/Engine/Plugins/Experimental/UAF/UAFStateTree/Source/UAFStateTreeEditor/Private/AnimNextStateTreeEditorModule.h b/Engine/Plugins/Experimental/UAF/UAFStateTree/Source/UAFStateTreeEditor/Private/AnimNextStateTreeEditorModule.h
--- a/Engine/Plugins/Experimental/UAF/UAFStateTree/Source/UAFStateTreeEditor/Private/AnimNextStateTreeEditorModule.h
+++ b/Engine/Plugins/Experimental/UAF/UAFStateTree/Source/UAFStateTreeEditor/Private/AnimNextStateTreeEditorModule.h
@@ -3,6 +3,7 @@
 #pragma once
 
 #include "IAnimNextStateTreeEditorModule.h"
+#include "UObject/TopLevelAssetPath.h"
 
 namespace UE::UAF::StateTree
 {
@@ -11,5 +12,8 @@ class FAnimNextStateTreeEditorModule : public IAnimNextStateTreeEditorModule
 public:	
 	virtual void StartupModule() override;
 	virtual void ShutdownModule() override;
+	
+private:
+	FTopLevelAssetPath StateTreeClassPath;
 };
 }
diff --git a/Engine/Plugins/Experimental/UAF/UAFStateTree/Source/UAFStateTreeEditor/UAFStateTreeEditor.Build.cs b/Engine/Plugins/Experimental/UAF/UAFStateTree/Source/UAFStateTreeEditor/UAFStateTreeEditor.Build.cs
--- a/Engine/Plugins/Experimental/UAF/UAFStateTree/Source/UAFStateTreeEditor/UAFStateTreeEditor.Build.cs
+++ b/Engine/Plugins/Experimental/UAF/UAFStateTree/Source/UAFStateTreeEditor/UAFStateTreeEditor.Build.cs
@@ -21,7 +21,9 @@ namespace UnrealBuildTool.Rules
 				    "RigVM",
 				    "StructUtilsEditor",
 				    "PropertyEditor",
-					"ToolMenus"
+					"ToolMenus",
+					"UAFAnimNodeEditor",
+					"UAFAnimNode"
 			    });
 
 			PublicDependencyModuleNames.AddRange(
diff --git a/Engine/Plugins/Experimental/UAF/UAFStateTree/Source/UAFStateTreeUncookedOnly/UAFStateTreeUncookedOnly.Build.cs b/Engine/Plugins/Experimental/UAF/UAFStateTree/Source/UAFStateTreeUncookedOnly/UAFStateTreeUncookedOnly.Build.cs
--- a/Engine/Plugins/Experimental/UAF/UAFStateTree/Source/UAFStateTreeUncookedOnly/UAFStateTreeUncookedOnly.Build.cs
+++ b/Engine/Plugins/Experimental/UAF/UAFStateTree/Source/UAFStateTreeUncookedOnly/UAFStateTreeUncookedOnly.Build.cs
@@ -27,6 +27,7 @@ namespace UnrealBuildTool.Rules
 					"UAFUncookedOnly",
 					"UAFStateTree",
 					"UAFAnimGraph",
+					"UAFAnimNode",
 					"UAFAnimGraphUncookedOnly"
 				}
 			);
