Fix missing UAFAnimNode dependency in UAFStateTRee

#Synced-CL 51417719

diff --git a/Engine/Plugins/Experimental/UAF/UAFStateTree/UAFStateTree.uplugin b/Engine/Plugins/Experimental/UAF/UAFStateTree/UAFStateTree.uplugin
--- a/Engine/Plugins/Experimental/UAF/UAFStateTree/UAFStateTree.uplugin
+++ b/Engine/Plugins/Experimental/UAF/UAFStateTree/UAFStateTree.uplugin
@@ -32,6 +32,10 @@
 		{
 			"Name": "UAFAnimGraph",
 			"Enabled": true
+		},
+		{
+			"Name": "UAFAnimNode",
+			"Enabled": true
 		}
 	],
 	"Modules" :
