Fix UAF ST Async compile crashes

#tests PIE, editor
#Synced-CL 51463347

diff --git a/Engine/Plugins/Experimental/UAF/UAFStateTree/Source/UAFStateTree/Internal/AnimNextStateTreeSchema.h b/Engine/Plugins/Experimental/UAF/UAFStateTree/Source/UAFStateTree/Internal/AnimNextStateTreeSchema.h
--- a/Engine/Plugins/Experimental/UAF/UAFStateTree/Source/UAFStateTree/Internal/AnimNextStateTreeSchema.h
+++ b/Engine/Plugins/Experimental/UAF/UAFStateTree/Source/UAFStateTree/Internal/AnimNextStateTreeSchema.h
@@ -26,6 +26,7 @@ public:
 #if WITH_EDITOR
 	virtual bool AllowUtilityConsiderations() const override { return false; }
 	virtual bool AllowEvaluators() const override { return false; }
+	virtual bool AllowQueuedCompilation() const override { return false; }
 #endif // WITH_EDITOR
 protected:
 	/** List of named external data required by schema and provided to the state tree through the execution context. */
diff --git a/Engine/Plugins/Experimental/UAF/UAFStateTree/Source/UAFStateTree/Private/AnimStateTreeTrait.cpp b/Engine/Plugins/Experimental/UAF/UAFStateTree/Source/UAFStateTree/Private/AnimStateTreeTrait.cpp
--- a/Engine/Plugins/Experimental/UAF/UAFStateTree/Source/UAFStateTree/Private/AnimStateTreeTrait.cpp
+++ b/Engine/Plugins/Experimental/UAF/UAFStateTree/Source/UAFStateTree/Private/AnimStateTreeTrait.cpp
@@ -100,6 +100,13 @@ namespace UE::UAF
 			StateTreeExecutionContext.SetOuterTraceId(OwnerGraphInstance.GetUniqueId());
 
 #if ENABLE_ANIM_DEBUG 
+			if (!InstanceData->StateTree->IsReadyToRun())
+			{
+				UE_LOG(LogAnimation, Error, TEXT("StateTree cannot be ran, may have failed compilation StateTree: %s"), *InstanceData->StateTree->GetOuter()->GetFName().ToString());
+				StateTreeExecutionContext.Stop(EStateTreeRunStatus::Failed);
+				return;
+			}
+
 			if (!OwnerGraphInstance.LayoutMatches(InstanceData->StateTree->GetDefaultParameters()))
 			{
 				UE_LOG(LogAnimation, Error, TEXT("Anim StateTree Parameter Layout Mismatch: %s - StateTree: %s"), *Owner->GetFName().ToString(), *InstanceData->StateTree->GetOuter()->GetFName().ToString());
@@ -198,7 +205,7 @@ namespace UE::UAF
 	void FStateTreeTrait::PreUpdate(FUpdateTraversalContext& Context, const TTraitBinding<IUpdate>& Binding, const FTraitUpdateState& TraitState) const
 	{
 		FInstanceData* InstanceData = Binding.GetInstanceData<FInstanceData>();
-		if (InstanceData->StateTree)
+		if (InstanceData->StateTree && InstanceData->StateTree->IsReadyToRun())
 		{
 			UObject* Owner = GetTransientPackage();
 #if ENABLE_ANIM_DEBUG 
