Add ability to directly set a UAF variable from a task in UAF ST

#tests editor, PIE
#Synced-CL 51775323

diff --git a/Engine/Plugins/Experimental/UAF/UAFStateTree/Source/UAFStateTree/Internal/Tasks/UAFStateTreeSetVariableTask.h b/Engine/Plugins/Experimental/UAF/UAFStateTree/Source/UAFStateTree/Internal/Tasks/UAFStateTreeSetVariableTask.h
new file mode 100644
--- /dev/null
+++ b/Engine/Plugins/Experimental/UAF/UAFStateTree/Source/UAFStateTree/Internal/Tasks/UAFStateTreeSetVariableTask.h
@@ -0,0 +1,56 @@
+// Copyright Epic Games, Inc. All Rights Reserved.
+
+#pragma once
+
+#include "CoreMinimal.h"
+#include "AnimNextStateTreeTypes.h"
+#include "StateTreeLinker.h"
+#include "StructUtils/PropertyBag.h"
+#include "UAFStateTreeContext.h"
+#include "Variables/AnimNextVariableReference.h"
+
+#include "UAFStateTreeSetVariableTask.generated.h"
+
+struct FStateTreeExecutionContext;
+struct FStateTreeTransitionResult;
+
+/** Instance data for FUAFStateTreeSetVariableTask. */
+USTRUCT()
+struct FUAFStateTreeSetVariableTaskInstanceData
+{
+	GENERATED_BODY()
+
+	/** Which UAF variable to write. */
+	UPROPERTY(EditAnywhere, Category = Variable)
+	FAnimNextVariableReference Variable;
+
+	/** Value to write. Automatically reshaped to match Variable's type when Variable changes. */
+	UPROPERTY(EditAnywhere, Category = Variable, Meta = (FixedLayout, ShowOnlyInnerProperties))
+	FInstancedPropertyBag Value;
+};
+
+/**
+ * StateTree task that writes a single UAF shared variable on state entry.
+ * Stack multiple instances of this task in one state to write multiple variables independently.
+ */
+USTRUCT(MinimalAPI, DisplayName = "UAF Set Variable")
+struct FUAFStateTreeSetVariableTask : public FAnimNextStateTreeTaskBase
+{
+	GENERATED_BODY()
+
+	using FInstanceDataType = FUAFStateTreeSetVariableTaskInstanceData;
+
+	//~ Begin FStateTreeTaskBase Interface
+	virtual bool Link(FStateTreeLinker& Linker) override;
+	virtual EStateTreeRunStatus EnterState(FStateTreeExecutionContext& Context, const FStateTreeTransitionResult& Transition) const override;
+	//~ End FStateTreeTaskBase Interface
+
+	//~ Begin FStateTreeNodeBase Interface
+	virtual const UStruct* GetInstanceDataType() const override { return FInstanceDataType::StaticStruct(); }
+#if WITH_EDITOR
+	virtual void PostEditInstanceDataChangeChainProperty(const FPropertyChangedChainEvent& PropertyChangedEvent, FStateTreeDataView InstanceDataView) override;
+#endif
+	//~ End FStateTreeNodeBase Interface
+
+	TStateTreeExternalDataHandle<FUAFStateTreeContext> ContextHandle;
+};
diff --git a/Engine/Plugins/Experimental/UAF/UAFStateTree/Source/UAFStateTree/Private/Tasks/UAFStateTreeSetVariableTask.cpp b/Engine/Plugins/Experimental/UAF/UAFStateTree/Source/UAFStateTree/Private/Tasks/UAFStateTreeSetVariableTask.cpp
new file mode 100644
--- /dev/null
+++ b/Engine/Plugins/Experimental/UAF/UAFStateTree/Source/UAFStateTree/Private/Tasks/UAFStateTreeSetVariableTask.cpp
@@ -0,0 +1,95 @@
+// Copyright Epic Games, Inc. All Rights Reserved.
+
+#include "Tasks/UAFStateTreeSetVariableTask.h"
+
+#include "Param/ParamType.h"
+#include "StateTreeExecutionContext.h"
+#include "StateTreeLinker.h"
+#include "UAFAssetInstance.h"
+#include "UAFStateTreeContext.h"
+
+#include UE_INLINE_GENERATED_CPP_BY_NAME(UAFStateTreeSetVariableTask)
+
+////////////////////////////////////////////////////////////////////////////////////////////////////////////////////////
+// FUAFStateTreeSetVariableTask
+
+bool FUAFStateTreeSetVariableTask::Link(FStateTreeLinker& Linker)
+{
+	Linker.LinkExternalData(ContextHandle);
+	return true;
+}
+
+EStateTreeRunStatus FUAFStateTreeSetVariableTask::EnterState(FStateTreeExecutionContext& Context, const FStateTreeTransitionResult& Transition) const
+{
+	const FInstanceDataType& InstanceData = Context.GetInstanceData<FInstanceDataType>(*this);
+
+	if (InstanceData.Variable.IsNone())
+	{
+		ensureMsgf(false, TEXT("UAF Set Variable: No Variable assigned. StateTree: '%s', State: '%s', Owner: '%s'"),
+			*Context.GetStateTree()->GetName(),
+			*Context.GetActiveStateName(),
+			*Context.GetOwner()->GetName());
+		return EStateTreeRunStatus::Failed;
+	}
+
+	if (!InstanceData.Value.IsValid() || InstanceData.Value.GetNumPropertiesInBag() != 1)
+	{
+		ensureMsgf(false, TEXT("UAF Set Variable: Value invalid '%s'. StateTree: '%s', State: '%s', Owner: '%s'"),
+			*InstanceData.Variable.GetName().ToString(),
+			*Context.GetStateTree()->GetName(),
+			*Context.GetActiveStateName(),
+			*Context.GetOwner()->GetName());
+		return EStateTreeRunStatus::Failed;
+	}
+
+	FUAFStateTreeContext& ExecContext = Context.GetExternalData(ContextHandle);
+	if (FUAFAssetInstance * VariablesOwner = ExecContext.GetVariablesOwner())
+	{
+		const UPropertyBag* BagStruct = InstanceData.Value.GetPropertyBagStruct();
+		TConstArrayView<FPropertyBagPropertyDesc> Descs = BagStruct->GetPropertyDescs();
+		const FPropertyBagPropertyDesc& Desc = Descs[0];
+
+		const uint8* BagMemory = InstanceData.Value.GetValue().GetMemory();
+		const uint8* PropMemory = Desc.CachedProperty->ContainerPtrToValuePtr<uint8>(BagMemory);
+
+		const FAnimNextParamType VarType = FAnimNextParamType::FromPropertyBagPropertyDesc(Desc);
+
+		if (EPropertyBagResult::Success == VariablesOwner->SetVariable(InstanceData.Variable, VarType, TConstArrayView<uint8>(PropMemory, Desc.CachedProperty->GetSize())))
+		{
+			return EStateTreeRunStatus::Running;
+		}
+	}
+
+	return EStateTreeRunStatus::Failed;
+}
+
+#if WITH_EDITOR
+void FUAFStateTreeSetVariableTask::PostEditInstanceDataChangeChainProperty(const FPropertyChangedChainEvent& PropertyChangedEvent, FStateTreeDataView InstanceDataView)
+{
+	FInstanceDataType* InstanceData = InstanceDataView.GetMutablePtr<FInstanceDataType>();
+	if (!InstanceData)
+	{
+		return;
+	}
+
+	// Only reshape the Value bag when Variable itself has changed
+	if (PropertyChangedEvent.GetMemberPropertyName() != GET_MEMBER_NAME_CHECKED(FUAFStateTreeSetVariableTaskInstanceData, Variable))
+	{
+		return;
+	}
+
+	// Resolve the FProperty for the selected variable to determine its type
+	const FProperty* VarProp = InstanceData->Variable.ResolveProperty();
+	if (!VarProp)
+	{
+		InstanceData->Value.Reset();
+		return;
+	}
+
+	// Reshape Value to a single property matching the variable's type, preserving any existing value
+	FPropertyBagPropertyDesc NewDesc(InstanceData->Variable.GetName(), VarProp);
+	FInstancedPropertyBag TempBag;
+	TempBag.AddProperties({ NewDesc });
+	InstanceData->Value.MigrateToNewBagInstance(TempBag);
+}
+#endif // WITH_EDITOR
