Fix build warnings from uninitialized UProperties

#Synced-CL 51932823

diff --git a/Engine/Plugins/Experimental/UAF/UAFStateTree/Source/UAFStateTree/Internal/AnimNode/UAFStateTreeNodeConditions.h b/Engine/Plugins/Experimental/UAF/UAFStateTree/Source/UAFStateTree/Internal/AnimNode/UAFStateTreeNodeConditions.h
--- a/Engine/Plugins/Experimental/UAF/UAFStateTree/Source/UAFStateTree/Internal/AnimNode/UAFStateTreeNodeConditions.h
+++ b/Engine/Plugins/Experimental/UAF/UAFStateTree/Source/UAFStateTree/Internal/AnimNode/UAFStateTreeNodeConditions.h
@@ -67,10 +67,10 @@ struct FUAFFloatCompareCondition : public FUAFVariableConditionBase
 	FAnimNextVariableReference Left;
 
 	UPROPERTY(EditAnywhere, Category = Default)
-	EUAFFloatCompareConditionType Compare;
+	EUAFFloatCompareConditionType Compare = EUAFFloatCompareConditionType::Less;
 
 	UPROPERTY(EditAnywhere, Category = Default)
-	double Right;
+	double Right = 0;
 };
 
 
@@ -87,7 +87,7 @@ struct FUAFEnumCompareCondition : public FUAFVariableConditionBase
 	FAnimNextVariableReference Left;
 
 	UPROPERTY(EditAnywhere, Category = Default)
-	uint8 Right;
+	uint8 Right = 0;
 };
 
 }
\ No newline at end of file
