Fix incorrect ifdef

#Synced-CL 51932823

diff --git a/Engine/Plugins/Experimental/UAF/UAFStateTree/Source/UAFStateTree/Private/AnimNode/UAFStateTreeNode.cpp b/Engine/Plugins/Experimental/UAF/UAFStateTree/Source/UAFStateTree/Private/AnimNode/UAFStateTreeNode.cpp
--- a/Engine/Plugins/Experimental/UAF/UAFStateTree/Source/UAFStateTree/Private/AnimNode/UAFStateTreeNode.cpp
+++ b/Engine/Plugins/Experimental/UAF/UAFStateTree/Source/UAFStateTree/Private/AnimNode/UAFStateTreeNode.cpp
@@ -46,7 +46,7 @@ FUAFStateTreeNode::FUAFStateTreeNode(FUAFAnimGraphUpdateContext& Context, TObjec
 		FUAFStateTreeNodeContext UAFStateTreeContext(Context, this);
 		StateTreeExecutionContext.SetContextDataByName(UStateTreeAnimNextSchema::AnimStateTreeExecutionContextName, FStateTreeDataView(FUAFStateTreeContext::StaticStruct(), &UAFStateTreeContext));
 
-#if OBJECT_TRACE_ENABLED
+#if UAF_TRACE_ENABLED
 		UAF_TRACE_ANIMNODE(Context, this);
 		StateTreeExecutionContext.SetOuterTraceId(DebugInstanceId);
 #endif
