Fix crash attempting to manually create UAF ST by hiding UAF ST Schema

#tests editor
#Synced-CL 52261262

diff --git a/Engine/Plugins/Experimental/UAF/UAFStateTree/Source/UAFStateTree/Internal/AnimNextStateTreeSchema.h b/Engine/Plugins/Experimental/UAF/UAFStateTree/Source/UAFStateTree/Internal/AnimNextStateTreeSchema.h
--- a/Engine/Plugins/Experimental/UAF/UAFStateTree/Source/UAFStateTree/Internal/AnimNextStateTreeSchema.h
+++ b/Engine/Plugins/Experimental/UAF/UAFStateTree/Source/UAFStateTree/Internal/AnimNextStateTreeSchema.h
@@ -7,7 +7,7 @@
 #include "AnimNextStateTreeSchema.generated.h"
 
 
-UCLASS()
+UCLASS(HideDropDown)
 class UAFSTATETREE_API UStateTreeAnimNextSchema : public UStateTreeSchema
 {
 	GENERATED_BODY()
